Moved duplicated code from MeshUtilities and MeshDescriptionOperations (FLayoutUV, FAllocator2D, FOverlappingCorners, etc.) into a new single module MeshUtilitiesCommon.

Add a generic opaque mesh view interface FLayoutUV::IMeshView to abstract mesh data access and allow FLayoutUV to be used with any mesh type in any module.
Replaced few instances of using an old version of overlapping corners data structure (multi-map) with new specialized FOverlappingCorners container.

diff --git a/Engine/Plugins/Experimental/GeometryCache/Source/GeometryCache/GeometryCache.Build.cs b/Engine/Plugins/Experimental/GeometryCache/Source/GeometryCache/GeometryCache.Build.cs
--- a/Engine/Plugins/Experimental/GeometryCache/Source/GeometryCache/GeometryCache.Build.cs
+++ b/Engine/Plugins/Experimental/GeometryCache/Source/GeometryCache/GeometryCache.Build.cs
@@ -20,6 +20,12 @@ public class GeometryCache : ModuleRules
 			}
 		);
 
+        PrivateDependencyModuleNames.AddRange(
+            new string[] {
+                "MeshUtilitiesCommon",
+            }
+        );
+
         PublicIncludePathModuleNames.Add("TargetPlatform");
 
         if (Target.bBuildEditor)
diff --git a/Engine/Plugins/Experimental/GeometryCache/Source/GeometryCache/Private/GeometryCachePreprocessor.cpp b/Engine/Plugins/Experimental/GeometryCache/Source/GeometryCache/Private/GeometryCachePreprocessor.cpp
--- a/Engine/Plugins/Experimental/GeometryCache/Source/GeometryCache/Private/GeometryCachePreprocessor.cpp
+++ b/Engine/Plugins/Experimental/GeometryCache/Source/GeometryCache/Private/GeometryCachePreprocessor.cpp
@@ -6,6 +6,7 @@
 #include "Modules/ModuleManager.h"
 #include "MeshUtilities.h"
 #include "MeshBuild.h"
+#include "OverlappingCorners.h"
 #include "GeometryCacheTrackStreamable.h"
 
 void FCodecGeometryCachePreprocessor::AddMeshSample(const FGeometryCacheMeshData& MeshData, const float SampleTime, bool bSameTopologyAsPrevious)
diff --git a/Engine/Source/Developer/MeshBuilder/MeshBuilder.Build.cs b/Engine/Source/Developer/MeshBuilder/MeshBuilder.Build.cs
--- a/Engine/Source/Developer/MeshBuilder/MeshBuilder.Build.cs
+++ b/Engine/Source/Developer/MeshBuilder/MeshBuilder.Build.cs
@@ -18,7 +18,8 @@ namespace UnrealBuildTool.Rules
                     "RenderCore",
                     "MeshDescriptionOperations",
                     "MeshReductionInterface",
-                    "RawMesh"
+                    "RawMesh",
+                    "MeshUtilitiesCommon",
                 }
 			);
 
diff --git a/Engine/Source/Developer/MeshBuilder/Private/MeshDescriptionHelper.cpp b/Engine/Source/Developer/MeshBuilder/Private/MeshDescriptionHelper.cpp
--- a/Engine/Source/Developer/MeshBuilder/Private/MeshDescriptionHelper.cpp
+++ b/Engine/Source/Developer/MeshBuilder/Private/MeshDescriptionHelper.cpp
@@ -132,12 +132,12 @@ void FMeshDescriptionHelper::GetRenderMeshDescription(UObject* Owner, const FMes
 			BuildSettings->SrcLightmapIndex,
 			BuildSettings->DstLightmapIndex,
 			BuildSettings->MinLightmapResolution,
-			(FMeshDescriptionOperations::ELightmapUVVersion)(StaticMesh->LightmapUVVersion),
+			(ELightmapUVVersion)(StaticMesh->LightmapUVVersion),
 			OverlappingCorners);
 	}
 }
 
-void FMeshDescriptionHelper::ReduceLOD(const FMeshDescription& BaseMesh, FMeshDescription& DestMesh, const FMeshReductionSettings& ReductionSettings, const TMultiMap<int32, int32>& InOverlappingCorners, float &OutMaxDeviation)
+void FMeshDescriptionHelper::ReduceLOD(const FMeshDescription& BaseMesh, FMeshDescription& DestMesh, const FMeshReductionSettings& ReductionSettings, const FOverlappingCorners& InOverlappingCorners, float &OutMaxDeviation)
 {
 	IMeshReductionManagerModule& MeshReductionModule = FModuleManager::Get().LoadModuleChecked<IMeshReductionManagerModule>("MeshReductionInterface");
 	IMeshReduction* MeshReduction = MeshReductionModule.GetStaticMeshReductionInterface();
diff --git a/Engine/Source/Developer/MeshBuilder/Private/MeshDescriptionHelper.h b/Engine/Source/Developer/MeshBuilder/Private/MeshDescriptionHelper.h
--- a/Engine/Source/Developer/MeshBuilder/Private/MeshDescriptionHelper.h
+++ b/Engine/Source/Developer/MeshBuilder/Private/MeshDescriptionHelper.h
@@ -4,6 +4,7 @@
 #include "CoreMinimal.h"
 #include "CoreTypes.h"
 #include "Logging/LogMacros.h"
+#include "OverlappingCorners.h"
 
 DECLARE_LOG_CATEGORY_EXTERN(LogMeshBuilder, Log, All);
 
@@ -22,11 +23,11 @@ public:
 	//Build a render mesh description with the BuildSettings. This will update the InRenderMeshDescription ptr content
 	void GetRenderMeshDescription(UObject* Owner, const FMeshDescription& InOriginalMeshDescription, FMeshDescription& OutRenderMeshDescription);
 
-	void ReduceLOD(const FMeshDescription& BaseMesh, FMeshDescription& DestMesh, const struct FMeshReductionSettings& ReductionSettings, const TMultiMap<int32, int32>& InOverlappingCorners, float &OutMaxDeviation);
+	void ReduceLOD(const FMeshDescription& BaseMesh, FMeshDescription& DestMesh, const struct FMeshReductionSettings& ReductionSettings, const FOverlappingCorners& InOverlappingCorners, float &OutMaxDeviation);
 
 	void FindOverlappingCorners(const FMeshDescription& MeshDescription, float ComparisonThreshold);
 
-	const TMultiMap<int32, int32>& GetOverlappingCorners() const { return OverlappingCorners; }
+	const FOverlappingCorners& GetOverlappingCorners() const { return OverlappingCorners; }
 
 private:
 
@@ -37,7 +38,7 @@ private:
 	//PRIVATE class members
 
 	FMeshBuildSettings* BuildSettings;
-	TMultiMap<int32, int32> OverlappingCorners;
+	FOverlappingCorners OverlappingCorners;
 
 	
 	//////////////////////////////////////////////////////////////////////////
diff --git a/Engine/Source/Developer/MeshBuilder/Private/StaticMeshBuilder.cpp b/Engine/Source/Developer/MeshBuilder/Private/StaticMeshBuilder.cpp
--- a/Engine/Source/Developer/MeshBuilder/Private/StaticMeshBuilder.cpp
+++ b/Engine/Source/Developer/MeshBuilder/Private/StaticMeshBuilder.cpp
@@ -27,7 +27,7 @@ void BuildVertexBuffer(
 	, TArray<int32>& OutWedgeMap
 	, TArray<TArray<uint32> >& OutPerSectionIndices
 	, TArray< FStaticMeshBuildVertex >& StaticMeshBuildVertices
-	, const TMultiMap<int32, int32>& OverlappingCorners
+	, const FOverlappingCorners& OverlappingCorners
 	, float VertexComparisonThreshold
 	, TArray<int32>& RemapVerts);
 void BuildAllBufferOptimizations(struct FStaticMeshLODResources& StaticMeshLOD, const struct FMeshBuildSettings& LODBuildSettings, TArray< uint32 >& IndexBuffer, bool bNeeds32BitIndices, TArray< FStaticMeshBuildVertex >& StaticMeshBuildVertices);
@@ -128,7 +128,7 @@ bool FStaticMeshBuilder::Build(FStaticMeshRenderData& StaticMeshRenderData, USta
 		{
 			const int32 BaseLodIndex = 0;
 			float OverlappingThreshold = LODBuildSettings.bRemoveDegenerates ? THRESH_POINTS_ARE_SAME : 0.0f;
-			TMultiMap<int32, int32> OverlappingCorners;
+			FOverlappingCorners OverlappingCorners;
 			FMeshDescriptionOperations::FindOverlappingCorners(OverlappingCorners, MeshDescriptions[BaseLodIndex], OverlappingThreshold);
 
 			//Create a reduced mesh from the base LOD
@@ -308,7 +308,7 @@ void BuildVertexBuffer(
 	, TArray<int32>& OutWedgeMap
 	, TArray<TArray<uint32> >& OutPerSectionIndices
 	, TArray< FStaticMeshBuildVertex >& StaticMeshBuildVertices
-	, const TMultiMap<int32, int32>& OverlappingCorners
+	, const FOverlappingCorners& OverlappingCorners
 	, float VertexComparisonThreshold
 	, TArray<int32>& RemapVerts)
 {
@@ -330,7 +330,6 @@ void BuildVertexBuffer(
 	{
 		RemapIndex = INDEX_NONE;
 	}
-	TArray<int32> DupVerts;
 
 	TPolygonGroupAttributesConstRef<FName> PolygonGroupImportedMaterialSlotNames = MeshDescription.PolygonGroupAttributes().GetAttributesRef<FName>(MeshAttribute::PolygonGroup::ImportedMaterialSlotName);
 	TVertexAttributesConstRef<FVector> VertexPositions = MeshDescription.VertexAttributes().GetAttributesRef<FVector>( MeshAttribute::Vertex::Position );
@@ -434,9 +433,8 @@ void BuildVertexBuffer(
 					
 
 				//Never add duplicated vertex instance
-				DupVerts.Reset();
-				OverlappingCorners.MultiFind(VertexInstanceValue, DupVerts);
-				DupVerts.Sort();
+				const TArray<int32>& DupVerts = OverlappingCorners.FindIfOverlapping(VertexInstanceValue);
+
 				int32 Index = INDEX_NONE;
 				for (int32 k = 0; k < DupVerts.Num(); k++)
 				{
diff --git a/Engine/Source/Developer/MeshDescriptionOperations/MeshDescriptionOperations.Build.cs b/Engine/Source/Developer/MeshDescriptionOperations/MeshDescriptionOperations.Build.cs
--- a/Engine/Source/Developer/MeshDescriptionOperations/MeshDescriptionOperations.Build.cs
+++ b/Engine/Source/Developer/MeshDescriptionOperations/MeshDescriptionOperations.Build.cs
@@ -16,6 +16,7 @@ namespace UnrealBuildTool.Rules
                     "CoreUObject",
                     "RenderCore",
                     "MeshDescription",
+                    "MeshUtilitiesCommon",
                     "RawMesh",
                     "RHI"
                 }
diff --git a/Engine/Source/Developer/MeshDescriptionOperations/Private/Allocator2D.cpp b/Engine/Source/Developer/MeshDescriptionOperations/Private/Allocator2D.cpp
deleted file mode 100644
diff --git a/Engine/Source/Developer/MeshDescriptionOperations/Private/Allocator2D.h b/Engine/Source/Developer/MeshDescriptionOperations/Private/Allocator2D.h
deleted file mode 100644
diff --git a/Engine/Source/Developer/MeshDescriptionOperations/Private/DisjointSet.h b/Engine/Source/Developer/MeshDescriptionOperations/Private/DisjointSet.h
deleted file mode 100644
diff --git a/Engine/Source/Developer/MeshDescriptionOperations/Private/LayoutUV.cpp b/Engine/Source/Developer/MeshDescriptionOperations/Private/LayoutUV.cpp
deleted file mode 100644
diff --git a/Engine/Source/Developer/MeshDescriptionOperations/Private/LayoutUV.h b/Engine/Source/Developer/MeshDescriptionOperations/Private/LayoutUV.h
deleted file mode 100644
diff --git a/Engine/Source/Developer/MeshDescriptionOperations/Private/MeshDescriptionOperations.cpp b/Engine/Source/Developer/MeshDescriptionOperations/Private/MeshDescriptionOperations.cpp
--- a/Engine/Source/Developer/MeshDescriptionOperations/Private/MeshDescriptionOperations.cpp
+++ b/Engine/Source/Developer/MeshDescriptionOperations/Private/MeshDescriptionOperations.cpp
@@ -5,9 +5,10 @@
 #include "MeshDescription.h"
 #include "MeshAttributes.h"
 #include "RawMesh.h"
+#include "LayoutUV.h"
+#include "OverlappingCorners.h"
 #include "RenderUtils.h"
 #include "mikktspace.h"
-#include "LayoutUV.h"
 
 DEFINE_LOG_CATEGORY(LogMeshDescriptionOperations);
 
@@ -1003,16 +1004,18 @@ namespace MeshDescriptionMikktSpaceInterface
 	}
 }
 
-void FMeshDescriptionOperations::FindOverlappingCorners(TMultiMap<int32, int32>& OverlappingCorners, const FMeshDescription& MeshDescription, float ComparisonThreshold)
+void FMeshDescriptionOperations::FindOverlappingCorners(FOverlappingCorners& OutOverlappingCorners, const FMeshDescription& MeshDescription, float ComparisonThreshold)
 {
-	//Empty the old data
-	OverlappingCorners.Reset();
+	// @todo: this should be shared with FOverlappingCorners
 
 	const FVertexInstanceArray& VertexInstanceArray = MeshDescription.VertexInstances();
 	const FVertexArray& VertexArray = MeshDescription.Vertices();
 
 	const int32 NumWedges = VertexInstanceArray.Num();
 
+	// Empty the old data and reserve space for new
+	OutOverlappingCorners.Init(NumWedges);
+
 	// Create a list of vertex Z/index pairs
 	TArray<MeshDescriptionOperationNamespace::FIndexAndZ> VertIndexAndZ;
 	VertIndexAndZ.Reserve(NumWedges);
@@ -1041,21 +1044,118 @@ void FMeshDescriptionOperations::FindOverlappingCorners(TMultiMap<int32, int32>&
 
 			if (PositionA.Equals(PositionB, ComparisonThreshold))
 			{
-				OverlappingCorners.Add(VertIndexAndZ[i].Index, VertIndexAndZ[j].Index);
-				OverlappingCorners.Add(VertIndexAndZ[j].Index, VertIndexAndZ[i].Index);
+				OutOverlappingCorners.Add(VertIndexAndZ[i].Index, VertIndexAndZ[j].Index);
+				OutOverlappingCorners.Add(VertIndexAndZ[j].Index, VertIndexAndZ[i].Index);
 			}
 		}
 	}
+
+	OutOverlappingCorners.FinishAdding();
 }
 
+struct FLayoutUVMeshDescriptionView final : FLayoutUV::IMeshView
+{
+	FMeshDescription& MeshDescription;
+	const TVertexAttributeArray<FVector>& Positions;
+	const TVertexInstanceAttributeArray<FVector>& Normals;
+	const TVertexInstanceAttributeArray<FVector2D>& TexCoords;
+
+	const uint32 SrcChannel;
+	const uint32 DstChannel;
+
+	TVertexInstanceAttributeArray<FVector2D>* OutputTexCoords = nullptr;
+
+	uint32 NumIndices = 0;
+	TArray<int32> RemapVerts;
+	TArray<FVector2D> FlattenedTexCoords;
+
+	FLayoutUVMeshDescriptionView(FMeshDescription& InMeshDescription, uint32 InSrcChannel, uint32 InDstChannel) 
+		: MeshDescription(InMeshDescription)
+		, Positions(InMeshDescription.VertexAttributes().GetAttributes<FVector>(MeshAttribute::Vertex::Position))
+		, Normals(InMeshDescription.VertexInstanceAttributes().GetAttributes<FVector>(MeshAttribute::VertexInstance::Normal))
+		, TexCoords(MeshDescription.VertexInstanceAttributes().GetAttributes<FVector2D>(MeshAttribute::VertexInstance::TextureCoordinate, InSrcChannel))
+		, SrcChannel(InSrcChannel)
+		, DstChannel(InDstChannel)
+	{
+		uint32 NumTris = 0;
+		for (const FPolygonID PolygonID : MeshDescription.Polygons().GetElementIDs())
+		{
+			NumTris += MeshDescription.GetPolygonTriangles(PolygonID).Num();
+		}
+
+		NumIndices = NumTris * 3;
+
+		FlattenedTexCoords.SetNumUninitialized(NumIndices);
+		RemapVerts.SetNumUninitialized(NumIndices);
+
+		int32 WedgeIndex = 0;
+
+		for (const FPolygonID PolygonID : MeshDescription.Polygons().GetElementIDs())
+		{
+			const TArray<FMeshTriangle>& Triangles = MeshDescription.GetPolygonTriangles(PolygonID);
+			for (const FMeshTriangle MeshTriangle : Triangles)
+			{
+				for (int32 Corner = 0; Corner < 3; ++Corner)
+				{
+					const FVertexInstanceID VertexInstanceID = MeshTriangle.GetVertexInstanceID(Corner);
+
+					FlattenedTexCoords[WedgeIndex] = TexCoords[VertexInstanceID];
+					RemapVerts[WedgeIndex] = VertexInstanceID.GetValue();
+					++WedgeIndex;
+				}
+			}
+		}
+	}
+
+	uint32 GetNumIndices() const override { return NumIndices; }
+
+	FVector GetPosition(uint32 Index) const override
+	{ 
+		FVertexInstanceID VertexInstanceID(RemapVerts[Index]);
+		FVertexID VertexID = MeshDescription.GetVertexInstanceVertex(VertexInstanceID);
+		return Positions[VertexID];
+	}
+
+	FVector GetNormal(uint32 Index) const override
+	{ 
+		FVertexInstanceID VertexInstanceID(RemapVerts[Index]);
+		return Normals[VertexInstanceID];
+	}
+
+	FVector2D GetInputTexcoord(uint32 Index) const override
+	{
+		return FlattenedTexCoords[Index];
+	}
+
+	void InitOutputTexcoords(uint32 Num) override
+	{
+		// If current DstChannel is out of range of the number of UVs defined by the mesh description, change the index count accordingly
+		const uint32 NumUVs = MeshDescription.VertexInstanceAttributes().GetAttributeIndexCount<FVector2D>(MeshAttribute::VertexInstance::TextureCoordinate);
+		if (DstChannel >= NumUVs)
+		{
+			MeshDescription.VertexInstanceAttributes().SetAttributeIndexCount<FVector2D>(MeshAttribute::VertexInstance::TextureCoordinate, DstChannel + 1);
+			ensure(false);	// not expecting it to get here
+		}
+
+		OutputTexCoords = &MeshDescription.VertexInstanceAttributes().GetAttributes<FVector2D>(MeshAttribute::VertexInstance::TextureCoordinate, DstChannel);
+	}
+
+	void SetOutputTexcoord(uint32 Index, const FVector2D& Value) override
+	{
+		const FVertexInstanceID VertexInstanceID(RemapVerts[Index]);
+		(*OutputTexCoords)[VertexInstanceID] = Value;
+	}
+};
+
 void FMeshDescriptionOperations::CreateLightMapUVLayout(FMeshDescription& MeshDescription,
 	int32 SrcLightmapIndex,
 	int32 DstLightmapIndex,
 	int32 MinLightmapResolution,
 	ELightmapUVVersion LightmapUVVersion,
-	const TMultiMap<int32, int32>& OverlappingCorners)
+	const FOverlappingCorners& OverlappingCorners)
 {
-	MeshDescriptionOp::FLayoutUV Packer(MeshDescription, SrcLightmapIndex, DstLightmapIndex, MinLightmapResolution);
+	FLayoutUVMeshDescriptionView MeshDescriptionView(MeshDescription, SrcLightmapIndex, DstLightmapIndex);
+	FLayoutUV Packer(MeshDescriptionView, MinLightmapResolution);
 	Packer.SetVersion(LightmapUVVersion);
 
 	Packer.FindCharts(OverlappingCorners);
@@ -1072,11 +1172,12 @@ bool FMeshDescriptionOperations::GenerateUniqueUVsForStaticMesh(const FMeshDescr
 	FMeshDescription DuplicateMeshDescription(MeshDescription);
 	// Find overlapping corners for UV generator. Allow some threshold - this should not produce any error in a case if resulting
 	// mesh will not merge these vertices.
-	TMultiMap<int32, int32> OverlappingCorners;
+	FOverlappingCorners OverlappingCorners;
 	FindOverlappingCorners(OverlappingCorners, DuplicateMeshDescription, THRESH_POINTS_ARE_SAME);
 
 	// Generate new UVs
-	MeshDescriptionOp::FLayoutUV Packer(DuplicateMeshDescription, 0, 1, FMath::Clamp(TextureResolution / 4, 32, 512));
+	FLayoutUVMeshDescriptionView DuplicateMeshDescriptionView(DuplicateMeshDescription, 0, 1);
+	FLayoutUV Packer(DuplicateMeshDescriptionView, FMath::Clamp(TextureResolution / 4, 32, 512));
 	Packer.FindCharts(OverlappingCorners);
 
 	bool bPackSuccess = Packer.FindBestPacking();
diff --git a/Engine/Source/Developer/MeshDescriptionOperations/Public/MeshDescriptionOperations.h b/Engine/Source/Developer/MeshDescriptionOperations/Public/MeshDescriptionOperations.h
--- a/Engine/Source/Developer/MeshDescriptionOperations/Public/MeshDescriptionOperations.h
+++ b/Engine/Source/Developer/MeshDescriptionOperations/Public/MeshDescriptionOperations.h
@@ -4,9 +4,10 @@
 
 #include "CoreMinimal.h"
 
-
 class FMeshDescription;
 struct FRawMesh;
+struct FOverlappingCorners;
+enum class ELightmapUVVersion : int32;
 
 DECLARE_LOG_CATEGORY_EXTERN(LogMeshDescriptionOperations, Log, All);
 
@@ -25,14 +26,6 @@ public:
 		UseMikkTSpace = 0x4,
 	};
 
-	enum class ELightmapUVVersion : int32
-	{
-		BitByBit = 0,
-		Segments = 1,
-		SmallChartPacking = 2,
-		Latest = SmallChartPacking
-	};
-
 	/** Convert this mesh description into the old FRawMesh format. */
 	static void ConvertToRawMesh(const FMeshDescription& SourceMeshDescription, FRawMesh& DestinationRawMesh, const TMap<FName, int32>& MaterialMap);
 
@@ -52,14 +45,14 @@ public:
 	static void CreateMikktTangents(FMeshDescription& MeshDescription, ETangentOptions TangentOptions);
 
 	/** Find all overlapping vertex using the threshold in the mesh description. */
-	static void FindOverlappingCorners(TMultiMap<int32, int32>& OverlappingCorners, const FMeshDescription& MeshDescription, float ComparisonThreshold);
+	static void FindOverlappingCorners(FOverlappingCorners& OverlappingCorners, const FMeshDescription& MeshDescription, float ComparisonThreshold);
 
 	static void CreateLightMapUVLayout(FMeshDescription& MeshDescription,
 		int32 SrcLightmapIndex,
 		int32 DstLightmapIndex,
 		int32 MinLightmapResolution,
 		ELightmapUVVersion LightmapUVVersion,
-		const TMultiMap<int32, int32>& OverlappingCorners);
+		const FOverlappingCorners& OverlappingCorners);
 
 	/** Create some UVs from the specified mesh description data. */
 	static bool GenerateUniqueUVsForStaticMesh(const FMeshDescription& MeshDescription, int32 TextureResolution, TArray<FVector2D>& OutTexCoords);
diff --git a/Engine/Source/Developer/MeshMergeUtilities/MeshMergeUtilities.Build.cs b/Engine/Source/Developer/MeshMergeUtilities/MeshMergeUtilities.Build.cs
--- a/Engine/Source/Developer/MeshMergeUtilities/MeshMergeUtilities.Build.cs
+++ b/Engine/Source/Developer/MeshMergeUtilities/MeshMergeUtilities.Build.cs
@@ -24,6 +24,7 @@ public class MeshMergeUtilities : ModuleRules
                 "StaticMeshEditor",
                 "SkeletalMeshEditor",
                 "MaterialBaking",
+                "MeshUtilitiesCommon",
             }
 		);
         
diff --git a/Engine/Source/Developer/MeshReductionInterface/Public/IMeshReductionInterfaces.h b/Engine/Source/Developer/MeshReductionInterface/Public/IMeshReductionInterfaces.h
--- a/Engine/Source/Developer/MeshReductionInterface/Public/IMeshReductionInterfaces.h
+++ b/Engine/Source/Developer/MeshReductionInterface/Public/IMeshReductionInterfaces.h
@@ -45,7 +45,7 @@ public:
 		FMeshDescription& OutReducedMesh,
 		float& OutMaxDeviation,
 		const FMeshDescription& InMesh,
-		const TMultiMap<int32, int32>& InOverlappingCorners,
+		const FOverlappingCorners& InOverlappingCorners,
 		const struct FMeshReductionSettings& ReductionSettings
 	) = 0;
 
diff --git a/Engine/Source/Developer/MeshSimplifier/Private/QuadricMeshReduction.cpp b/Engine/Source/Developer/MeshSimplifier/Private/QuadricMeshReduction.cpp
--- a/Engine/Source/Developer/MeshSimplifier/Private/QuadricMeshReduction.cpp
+++ b/Engine/Source/Developer/MeshSimplifier/Private/QuadricMeshReduction.cpp
@@ -7,6 +7,7 @@
 #include "MeshBuild.h"
 #include "RawMesh.h"
 #include "MeshSimplify.h"
+#include "OverlappingCorners.h"
 #include "Templates/UniquePtr.h"
 #include "Features/IModularFeatures.h"
 #include "IMeshReductionInterfaces.h"
@@ -511,7 +512,7 @@ public:
 		FMeshDescription& OutReducedMesh,
 		float& OutMaxDeviation,
 		const FMeshDescription& InMesh,
-		const TMultiMap<int32, int32>& InOverlappingCorners,
+		const FOverlappingCorners& InOverlappingCorners,
 		const struct FMeshReductionSettings& ReductionSettings
 	) override
 	{
@@ -524,7 +525,6 @@ public:
 		TArray< uint32 >					Indexes;
 
 		TMap< int32, int32 > VertsMap;
-		TArray<int32> DupVerts;
 
 		int32 NumFaces = 0;
 		for (const FPolygonID PolygonID : InMesh.Polygons().GetElementIDs())
@@ -618,9 +618,7 @@ public:
 					// Make sure this vertex is valid from the start
 					NewVert.Correct();
 
-					DupVerts.Reset();
-					InOverlappingCorners.MultiFind(WedgeIndex, DupVerts);
-					DupVerts.Sort();
+					const TArray<int32>& DupVerts = InOverlappingCorners.FindIfOverlapping(WedgeIndex);
 
 					int32 Index = INDEX_NONE;
 					for (int32 k = 0; k < DupVerts.Num(); k++)
diff --git a/Engine/Source/Developer/MeshSimplifier/QuadricMeshReduction.Build.cs b/Engine/Source/Developer/MeshSimplifier/QuadricMeshReduction.Build.cs
--- a/Engine/Source/Developer/MeshSimplifier/QuadricMeshReduction.Build.cs
+++ b/Engine/Source/Developer/MeshSimplifier/QuadricMeshReduction.Build.cs
@@ -13,6 +13,7 @@ public class QuadricMeshReduction : ModuleRules
 		PrivateDependencyModuleNames.Add("RenderCore");
 		PrivateDependencyModuleNames.Add("RawMesh");
         PrivateDependencyModuleNames.Add("MeshDescription");
+        PrivateDependencyModuleNames.Add("MeshUtilitiesCommon");
 
         PrivateIncludePathModuleNames.AddRange(
         new string[] {
diff --git a/Engine/Source/Developer/MeshUtilities/MeshUtilities.Build.cs b/Engine/Source/Developer/MeshUtilities/MeshUtilities.Build.cs
--- a/Engine/Source/Developer/MeshUtilities/MeshUtilities.Build.cs
+++ b/Engine/Source/Developer/MeshUtilities/MeshUtilities.Build.cs
@@ -38,6 +38,7 @@ public class MeshUtilities : ModuleRules
 				"EditorStyle",
                 "GraphColor",
                 "MeshBuilder",
+                "MeshUtilitiesCommon",
             }
 		);
 
diff --git a/Engine/Source/Developer/MeshUtilities/Private/Allocator2D.cpp b/Engine/Source/Developer/MeshUtilities/Private/Allocator2D.cpp
deleted file mode 100644
diff --git a/Engine/Source/Developer/MeshUtilities/Private/DisjointSet.h b/Engine/Source/Developer/MeshUtilities/Private/DisjointSet.h
deleted file mode 100644
diff --git a/Engine/Source/Developer/MeshUtilities/Private/LayoutUV.cpp b/Engine/Source/Developer/MeshUtilities/Private/LayoutUV.cpp
deleted file mode 100644
diff --git a/Engine/Source/Developer/MeshUtilities/Private/MeshUtilities.cpp b/Engine/Source/Developer/MeshUtilities/Private/MeshUtilities.cpp
--- a/Engine/Source/Developer/MeshUtilities/Private/MeshUtilities.cpp
+++ b/Engine/Source/Developer/MeshUtilities/Private/MeshUtilities.cpp
@@ -88,6 +88,8 @@
 #include "DetailCategoryBuilder.h"
 #include "IDetailPropertyRow.h"
 #include "DetailWidgetRow.h"
+#include "OverlappingCorners.h"
+#include "MeshUtilitiesCommon.h"
 
 #if WITH_EDITOR
 #include "Editor.h"
@@ -826,29 +828,6 @@ void FMeshUtilities::BuildSkeletalModelFromChunks(FSkeletalMeshLODModel& LODMode
 Common functionality.
 ------------------------------------------------------------------------------*/
 
-/** Helper struct for building acceleration structures. */
-struct FIndexAndZ
-{
-	float Z;
-	int32 Index;
-
-	/** Default constructor. */
-	FIndexAndZ() {}
-
-	/** Initialization constructor. */
-	FIndexAndZ(int32 InIndex, FVector V)
-	{
-		Z = 0.30f * V.X + 0.33f * V.Y + 0.37f * V.Z;
-		Index = InIndex;
-	}
-};
-
-/** Sorting function for vertex Z/index pairs. */
-struct FCompareIndexAndZ
-{
-	FORCEINLINE bool operator()(FIndexAndZ const& A, FIndexAndZ const& B) const { return A.Z < B.Z; }
-};
-
 static int32 ComputeNumTexCoords(FRawMesh const& RawMesh, int32 MaxSupportedTexCoords)
 {
 	int32 NumWedges = RawMesh.WedgeIndices.Num();
@@ -864,20 +843,6 @@ static int32 ComputeNumTexCoords(FRawMesh const& RawMesh, int32 MaxSupportedTexC
 	return FMath::Min(NumTexCoords, MaxSupportedTexCoords);
 }
 
-/**
-* Returns true if the specified points are about equal
-*/
-inline bool PointsEqual(const FVector& V1, const FVector& V2, float ComparisonThreshold)
-{
-	if (FMath::Abs(V1.X - V2.X) > ComparisonThreshold
-		|| FMath::Abs(V1.Y - V2.Y) > ComparisonThreshold
-		|| FMath::Abs(V1.Z - V2.Z) > ComparisonThreshold)
-	{
-		return false;
-	}
-	return true;
-}
-
 static inline FVector GetPositionForWedge(FRawMesh const& Mesh, int32 WedgeIndex)
 {
 	int32 VertexIndex = Mesh.WedgeIndices[WedgeIndex];
@@ -1197,114 +1162,12 @@ static void ComputeTriangleTangents(
 	check(TriangleTangentZ.Num() == NumTriangles);*/
 }
 
-void FOverlappingCorners::Init(int32 NumIndices)
-{
-	Arrays.Reset();
-	Sets.Reset();
-	bFinishedAdding = false;
-
-	IndexBelongsTo.Reset(NumIndices);
-	IndexBelongsTo.AddUninitialized(NumIndices);
-	FMemory::Memset(IndexBelongsTo.GetData(), 0xFF, NumIndices * sizeof(int32));
-}
-
-void FOverlappingCorners::Add(int32 Key, int32 Value)
-{
-	check(Key != Value);
-	check(bFinishedAdding == false);
-
-	int32 ContainerIndex = IndexBelongsTo[Key];
-	if (ContainerIndex == INDEX_NONE)
-	{
-		ContainerIndex = Arrays.Num();
-		TArray<int32>& Container = Arrays.AddDefaulted_GetRef();
-		Container.Reserve(6);
-		Container.Add(Key);
-		Container.Add(Value);
-		IndexBelongsTo[Key] = ContainerIndex;
-		IndexBelongsTo[Value] = ContainerIndex;
-	}
-	else
-	{
-		IndexBelongsTo[Value] = ContainerIndex;
-
-		TArray<int32>& ArrayContainer = Arrays[ContainerIndex];
-		if (ArrayContainer.Num() == 1)
-		{
-			// Container is a set
-			Sets[ArrayContainer.Last()].Add(Value);
-		}
-		else
-		{
-			// Container is an array
-			ArrayContainer.AddUnique(Value);
-
-			// Change container into set when one vertex is shared by large number of triangles
-			if (ArrayContainer.Num() > 12)
-			{
-				int32 SetIndex = Sets.Num();
-				TSet<int32>& Set = Sets.AddDefaulted_GetRef();
-				Set.Append(ArrayContainer);
-
-				// Having one element means we are using a set
-				// An array will never have just 1 element normally because we add them as pairs
-				ArrayContainer.Reset(1);
-				ArrayContainer.Add(SetIndex);
-			}
-		}
-	}
-}
-
-void FOverlappingCorners::FinishAdding()
-{
-	check(bFinishedAdding == false);
-
-	for (TArray<int32>& Array : Arrays)
-	{
-		// Turn sets back into arrays for easier iteration code
-		// Also reduces peak memory later in the import process
-		if (Array.Num() == 1)
-		{
-			TSet<int32>& Set = Sets[Array.Last()];
-			Array.Reset(Set.Num());
-			for (int32 i : Set)
-			{
-				Array.Add(i);
-			}
-		}
-
-		// Sort arrays now to avoid sort multiple times
-		Array.Sort();
-	}
-
-	Sets.Empty();
-
-	bFinishedAdding = true;
-}
-
-uint32 FOverlappingCorners::GetAllocatedSize(void) const
-{
-	uint32 BaseMemoryAllocated = IndexBelongsTo.GetAllocatedSize() + Arrays.GetAllocatedSize() + Sets.GetAllocatedSize();
-
-	uint32 ArraysMemory = 0;
-	for (const TArray<int32>& ArrayIt : Arrays)
-	{
-		ArraysMemory += ArrayIt.GetAllocatedSize();
-	}
-
-	uint32 SetsMemory = 0;
-	for (const TSet<int32>& SetsIt : Sets)
-	{
-		SetsMemory += SetsIt.GetAllocatedSize();
-	}
-
-	return BaseMemoryAllocated + ArraysMemory + SetsMemory;
-}
-
 /**
 * Create a table that maps the corner of each face to its overlapping corners.
 * @param OutOverlappingCorners - Maps a corner index to the indices of all overlapping corners.
-* @param RawMesh - The mesh for which to compute overlapping corners.
+* @param InVertices - Triangle vertex positions for the mesh for which to compute overlapping corners.
+* @param InIndices - Triangle indices for the mesh for which to compute overlapping corners.
+* @param ComparisonThreshold - Positions are considered equal if all absolute differences between their X, Y and Z coordinates are less or equal to this value.
 */
 void FMeshUtilities::FindOverlappingCorners(
 	FOverlappingCorners& OutOverlappingCorners,
@@ -1312,47 +1175,14 @@ void FMeshUtilities::FindOverlappingCorners(
 	const TArray<uint32>& InIndices,
 	float ComparisonThreshold) const
 {
-	const int32 NumWedges = InIndices.Num();
-
-	// Create a list of vertex Z/index pairs
-	TArray<FIndexAndZ> VertIndexAndZ;
-	VertIndexAndZ.Reserve(NumWedges);
-	for (int32 WedgeIndex = 0; WedgeIndex < NumWedges; WedgeIndex++)
-	{
-		new(VertIndexAndZ)FIndexAndZ(WedgeIndex, InVertices[InIndices[WedgeIndex]]);
-	}
-
-	// Sort the vertices by z value
-	VertIndexAndZ.Sort(FCompareIndexAndZ());
-
-	OutOverlappingCorners.Init(NumWedges);
-
-	// Search for duplicates, quickly!
-	for (int32 i = 0; i < VertIndexAndZ.Num(); i++)
-	{
-		// only need to search forward, since we add pairs both ways
-		for (int32 j = i + 1; j < VertIndexAndZ.Num(); j++)
-		{
-			if (FMath::Abs(VertIndexAndZ[j].Z - VertIndexAndZ[i].Z) > ComparisonThreshold)
-				break; // can't be any more dups
-
-			const FVector& PositionA = InVertices[InIndices[VertIndexAndZ[i].Index]];
-			const FVector& PositionB = InVertices[InIndices[VertIndexAndZ[j].Index]];
-
-			if (PointsEqual(PositionA, PositionB, ComparisonThreshold))
-			{
-				OutOverlappingCorners.Add(VertIndexAndZ[i].Index, VertIndexAndZ[j].Index);
-			}
-		}
-	}
-
-	OutOverlappingCorners.FinishAdding();
+	OutOverlappingCorners = FOverlappingCorners(InVertices, InIndices, ComparisonThreshold);
 }
 
 /**
 * Create a table that maps the corner of each face to its overlapping corners.
 * @param OutOverlappingCorners - Maps a corner index to the indices of all overlapping corners.
 * @param RawMesh - The mesh for which to compute overlapping corners.
+* @param ComparisonThreshold - Positions are considered equal if all absolute differences between their X, Y and Z coordinates are less or equal to this value.
 */
 void FMeshUtilities::FindOverlappingCorners(
 	FOverlappingCorners& OutOverlappingCorners,
@@ -1360,7 +1190,7 @@ void FMeshUtilities::FindOverlappingCorners(
 	float ComparisonThreshold
 	) const
 {
-	FindOverlappingCorners(OutOverlappingCorners, RawMesh.VertexPositions, RawMesh.WedgeIndices, ComparisonThreshold);
+	OutOverlappingCorners = FOverlappingCorners(RawMesh.VertexPositions, RawMesh.WedgeIndices, ComparisonThreshold);
 }
 
 /**
@@ -2553,6 +2383,29 @@ void FMeshUtilities::CacheOptimizeVertexAndIndexBuffer(
 	}
 }
 
+struct FLayoutUVRawMeshView final : FLayoutUV::IMeshView
+{
+	FRawMesh& RawMesh;
+	const uint32 SrcChannel;
+	const uint32 DstChannel;
+	const bool bNormalsValid;
+
+	FLayoutUVRawMeshView(FRawMesh& InRawMesh, uint32 InSrcChannel, uint32 InDstChannel) 
+		: RawMesh(InRawMesh)
+		, SrcChannel(InSrcChannel)
+		, DstChannel(InDstChannel)
+		, bNormalsValid(InRawMesh.WedgeTangentZ.Num() == InRawMesh.WedgeTexCoords[InSrcChannel].Num())
+	{}
+
+	uint32     GetNumIndices() const override { return RawMesh.WedgeIndices.Num(); }
+	FVector    GetPosition(uint32 Index) const override { return RawMesh.GetWedgePosition(Index); }
+	FVector    GetNormal(uint32 Index) const override { return bNormalsValid ? RawMesh.WedgeTangentZ[Index] : FVector::ZeroVector; }
+	FVector2D  GetInputTexcoord(uint32 Index) const override { return RawMesh.WedgeTexCoords[SrcChannel][Index]; }
+
+	void      InitOutputTexcoords(uint32 Num) override { RawMesh.WedgeTexCoords[DstChannel].SetNumUninitialized( Num ); }
+	void      SetOutputTexcoord(uint32 Index, const FVector2D& Value) override { RawMesh.WedgeTexCoords[DstChannel][Index] = Value; }
+};
+
 class FStaticMeshUtilityBuilder
 {
 public:
@@ -2643,7 +2496,8 @@ public:
 						SrcModel.BuildSettings.SrcLightmapIndex = 0;
 					}
 
-					FLayoutUV Packer(&RawMesh, SrcModel.BuildSettings.SrcLightmapIndex, SrcModel.BuildSettings.DstLightmapIndex, SrcModel.BuildSettings.MinLightmapResolution);
+					FLayoutUVRawMeshView RawMeshView(RawMesh, SrcModel.BuildSettings.SrcLightmapIndex, SrcModel.BuildSettings.DstLightmapIndex);
+					FLayoutUV Packer(RawMeshView, SrcModel.BuildSettings.MinLightmapResolution);
 					Packer.SetVersion(LightmapUVVersion);
 
 					Packer.FindCharts(OverlappingCorners);
@@ -5526,7 +5380,8 @@ bool FMeshUtilities::GenerateUniqueUVsForSkeletalMesh(const FSkeletalMeshLODMode
 	OverlappingCorners.FinishAdding();
 
 	// Generate new UVs
-	FLayoutUV Packer(&TempMesh, 0, 1, FMath::Clamp(TextureResolution / 4, 32, 512));
+	FLayoutUVRawMeshView TempMeshView(TempMesh, 0, 1);
+	FLayoutUV Packer(TempMeshView, FMath::Clamp(TextureResolution / 4, 32, 512));
 	Packer.FindCharts(OverlappingCorners);
 
 	bool bPackSuccess = Packer.FindBestPacking();
@@ -6026,7 +5881,8 @@ bool FMeshUtilities::GenerateUniqueUVsForStaticMesh(const FRawMesh& RawMesh, int
 	FModuleManager::Get().LoadModuleChecked<IMeshUtilities>("MeshUtilities").FindOverlappingCorners(OverlappingCorners, TempMesh.VertexPositions, TempMesh.WedgeIndices, THRESH_POINTS_ARE_SAME);
 
 	// Generate new UVs
-	FLayoutUV Packer(&TempMesh, 0, 1, FMath::Clamp(TextureResolution / 4, 32, 512));
+	FLayoutUVRawMeshView TempMeshView(TempMesh, 0, 1);
+	FLayoutUV Packer(TempMeshView, FMath::Clamp(TextureResolution / 4, 32, 512));
 	Packer.FindCharts(OverlappingCorners);
 
 	bool bPackSuccess = Packer.FindBestPacking();
diff --git a/Engine/Source/Developer/MeshUtilities/Public/Allocator2D.h b/Engine/Source/Developer/MeshUtilities/Public/Allocator2D.h
deleted file mode 100644
diff --git a/Engine/Source/Developer/MeshUtilities/Public/LayoutUV.h b/Engine/Source/Developer/MeshUtilities/Public/LayoutUV.h
deleted file mode 100644
diff --git a/Engine/Source/Developer/MeshUtilities/Public/MeshUtilities.h b/Engine/Source/Developer/MeshUtilities/Public/MeshUtilities.h
--- a/Engine/Source/Developer/MeshUtilities/Public/MeshUtilities.h
+++ b/Engine/Source/Developer/MeshUtilities/Public/MeshUtilities.h
@@ -38,16 +38,6 @@ namespace ETangentOptions
 	};
 };
 
-enum class ELightmapUVVersion : int32
-{
-	BitByBit = 0,
-	Segments = 1,
-	SmallChartPacking = 2,
-	ScaleChartsOrderingFix = 3,
-	ChartJoiningLFix = 4,
-	Latest = ChartJoiningLFix
-};
-
 /**
 *	Contains the vertices that are most dominated by that bone. Vertices are in Bone space.
 *	Not used at runtime, but useful for fitting physics assets etc.
@@ -59,44 +49,7 @@ struct FBoneVertInfo
 	TArray<FVector>	Normals;
 };
 
-/**
-* Container to hold overlapping corners. For a vertex, lists all the overlapping vertices
-*/
-struct FOverlappingCorners
-{
-	FOverlappingCorners() : bFinishedAdding(false)
-	{
-	}
-
-	/* Resets, pre-allocates memory, marks all indices as not overlapping in preperation for calls to Add() */
-	void Init(int32 NumIndices);
-
-	/* Add overlapping indices pair */
-	void Add(int32 Key, int32 Value);
-
-	/* Sorts arrays, converts sets to arrays for sorting and to allow simple iterating code, prevents additional adding */
-	void FinishAdding();
-
-	/* Estimate memory allocated */
-	uint32 GetAllocatedSize(void) const;
-
-	/** 
-	* @return array of sorted overlapping indices including input 'Key', empty array for indices that have no overlaps.
-	*/
-	const TArray<int32>& FindIfOverlapping(int32 Key) const
-	{
-		check(bFinishedAdding);
-		int32 ContainerIndex = IndexBelongsTo[Key];
-		return (ContainerIndex != INDEX_NONE) ? Arrays[ContainerIndex] : EmptyArray;
-	}
-
-private:
-	TArray<int32> IndexBelongsTo;
-	TArray< TArray<int32> > Arrays;
-	TArray< TSet<int32> > Sets;
-	TArray<int32> EmptyArray;
-	bool bFinishedAdding;
-};
+struct FOverlappingCorners;
 
 class IMeshUtilities : public IModuleInterface
 {
diff --git a/Engine/Source/Developer/MeshUtilitiesCommon/MeshUtilitiesCommon.Build.cs b/Engine/Source/Developer/MeshUtilitiesCommon/MeshUtilitiesCommon.Build.cs
new file mode 100644
--- /dev/null
+++ b/Engine/Source/Developer/MeshUtilitiesCommon/MeshUtilitiesCommon.Build.cs
@@ -0,0 +1,21 @@
+// Copyright 1998-2018 Epic Games, Inc. All Rights Reserved.
+
+namespace UnrealBuildTool.Rules
+{
+    /**
+     * Common algorithms and data structures used by MeshUtilities, MeshDescriptionOperations and other mesh-related modules.
+     * This module must have minimal dependencies to allow it being used by various modules without introducing circular references.
+     */
+    public class MeshUtilitiesCommon : ModuleRules
+    {
+        public MeshUtilitiesCommon(ReadOnlyTargetRules Target) : base(Target)
+        {
+            PrivateDependencyModuleNames.AddRange(
+                new string[]
+                {
+                    "Core",
+                }
+            );
+        }
+    }
+}
diff --git a/Engine/Source/Developer/MeshUtilitiesCommon/Private/Allocator2D.cpp b/Engine/Source/Developer/MeshUtilitiesCommon/Private/Allocator2D.cpp
new file mode 100644
--- /dev/null
+++ b/Engine/Source/Developer/MeshUtilitiesCommon/Private/Allocator2D.cpp
@@ -0,0 +1,514 @@
+// Copyright 1998-2018 Epic Games, Inc. All Rights Reserved.
+
+#include "Allocator2D.h"
+
+FAllocator2D::FAllocator2D( uint32 InWidth, uint32 InHeight )
+	: Width( InWidth )
+	, Height( InHeight )
+	, LastRowFail( -1 )
+{
+	Pitch = ( Width + 63 ) / 64;
+	// alloc +1 to avoid buffer overrun
+	Bits = (uint64*)FMemory::Malloc( (Pitch * Height + 1) * sizeof(uint64) );
+
+	Rows.Reserve( Height );
+
+	for ( uint32 Y = 0; Y < Height; ++Y )
+	{
+		FRow Row;
+		Row.Index = Y;
+
+		Rows.Add( Row );
+	}
+
+	Clear();
+}
+
+FAllocator2D::~FAllocator2D()
+{
+	FMemory::Free( Bits );
+}
+
+FAllocator2D::FAllocator2D( const FAllocator2D& Other )
+	: Width( Other.Width )
+	, Height( Other.Height )
+	, Pitch( Other.Pitch )
+	, Rows( Other.Rows )
+	, LastRowFail( -1 )
+{
+	Bits = (uint64*)FMemory::Malloc( (Pitch * Height + 1) * sizeof(uint64) );
+	FMemory::Memcpy( Bits, Other.Bits, (Pitch * Height + 1) * sizeof(uint64) );
+}
+
+FAllocator2D& FAllocator2D::operator=( const FAllocator2D& Other )
+{
+	if ( Width != Other.Width || Height != Other.Height || Pitch != Other.Pitch )
+	{
+		Width = Other.Width;
+		Height = Other.Height;
+		Pitch = Other.Pitch;
+
+		FMemory::Free( Bits );
+		Bits = (uint64*)FMemory::Malloc( (Pitch * Height + 1) * sizeof(uint64) );
+	}
+
+	FMemory::Memcpy( Bits, Other.Bits, (Pitch * Height + 1) * sizeof(uint64) );
+
+	Rows = Other.Rows;
+
+	return *this;
+}
+
+void FAllocator2D::Clear()
+{
+	InitSegments();
+
+	FMemory::Memzero( Bits, (Pitch * Height + 1) * sizeof(uint64) );
+}
+
+bool FAllocator2D::Find( FRect& Rect )
+{
+	FRect TestRect = Rect;
+	for( TestRect.X = 0; TestRect.X <= Width - TestRect.W; TestRect.X++ )
+	{
+		for( TestRect.Y = 0; TestRect.Y <= Height - TestRect.H; TestRect.Y++ )
+		{
+			if( Test( TestRect ) )
+			{
+				Rect = TestRect;
+				return true;
+			}
+		}
+	}
+
+	return false;
+}
+
+bool FAllocator2D::FindBitByBit( FRect& Rect, const FAllocator2D& Other )
+{
+	FRect TestRect = Rect;
+	for( TestRect.X = 0; TestRect.X <= Width - TestRect.W; TestRect.X++ )
+	{
+		for( TestRect.Y = 0; TestRect.Y <= Height - TestRect.H; TestRect.Y++ )
+		{
+			if( Test( TestRect, Other ) )
+			{
+				Rect = TestRect;
+				return true;
+			}
+		}
+	}
+	
+	return false;
+}
+
+bool FAllocator2D::FindWithSegments( FRect& Rect, FRect BestRect, const FAllocator2D& Other )
+{
+	LastRowFail = -1;
+	FRect TestRect = Rect;
+
+	for( TestRect.Y = 0; TestRect.Y <= Height - TestRect.H; ++TestRect.Y )
+	{
+		uint32 FailedLength = 1;
+
+		for( TestRect.X = 0; TestRect.X <= Width - TestRect.W; TestRect.X += FailedLength )
+		{
+			if( TestRect.X + TestRect.Y * Height >= BestRect.X + BestRect.Y * Height )
+			{
+				// We've moved past the BestRect
+				return false;
+			}
+
+			if( TestAllRows( TestRect, Other, FailedLength ) )
+			{
+				Rect = TestRect;
+				return true;
+			}
+		}
+	}
+	
+	return false;
+}
+
+void FAllocator2D::Alloc( FRect Rect )
+{
+	for( uint32 y = Rect.Y; y < Rect.Y + Rect.H; y++ )
+	{
+		for( uint32 x = Rect.X; x < Rect.X + Rect.W; x++ )
+		{
+			SetBit( x, y );
+		}
+	}
+}
+
+void FAllocator2D::Alloc( FRect Rect, const FAllocator2D& Other )
+{
+	for( uint32 y = 0; y < Rect.H; y++ )
+	{
+		for( uint32 x = 0; x < Rect.W; x++ )
+		{
+			if( Other.GetBit( x, y ) )
+			{
+				SetBit( x + Rect.X, y + Rect.Y );
+			}
+		}
+	}
+
+	MergeSegments( Rect, Other );
+}
+
+bool FAllocator2D::TestAllRows( FRect Rect, const FAllocator2D& Other, uint32& FailedLength )
+{
+	FailedLength = 0.f;
+	bool bSuccess = true;
+
+	if ( LastRowFail > 0 )
+	{
+		const FRow& ThisRow = Rows[ Rect.Y + LastRowFail ];
+		const FRow& OtherRow = Other.Rows[ LastRowFail ];
+
+		if ( !TestRow( ThisRow, OtherRow, Rect, FailedLength ) )
+		{
+			return false;
+		}
+	}
+
+	LastRowFail = -1;
+
+	for ( uint32 y = 0; y < Rect.H; ++y )
+	{
+		const FRow& ThisRow = Rows[ Rect.Y + y ];
+		const FRow& OtherRow = Other.Rows[ y ];
+
+		uint32 CurrentFailedLength = 0;
+		bool bThisRowSuccess = TestRow( ThisRow, OtherRow, Rect, CurrentFailedLength );
+		bSuccess = bThisRowSuccess && bSuccess;
+
+		if ( !bThisRowSuccess && CurrentFailedLength > FailedLength )
+		{
+			LastRowFail = y;
+			FailedLength = CurrentFailedLength;
+		}
+
+		if ( FailedLength >= Width )
+		{
+			return false;
+		}
+	}
+
+	return bSuccess;
+}
+
+bool FAllocator2D::TestRow( const FRow& ThisRow, const FRow& OtherRow, FRect Rect, uint32& FailedLength )
+{
+	if ( ThisRow.LongestSegment < OtherRow.LongestSegment )
+	{
+		FailedLength = Width;
+		return false;
+	}
+
+	uint32 StartFreeSegmentIndex = 0;
+
+	for ( const FSegment& OtherUsedSegment : OtherRow.UsedSegments )
+	{
+		if ( OtherUsedSegment.StartPos >= Rect.W )
+		{
+			break;
+		}
+
+		bool bFoundSpaceForSegment = false;
+		bool bFoundFutureSpaceForSegment = false;
+
+		const uint32 StartPos = Rect.X + OtherUsedSegment.StartPos;
+		const uint32 EndPos = Rect.X + FMath::Min( OtherUsedSegment.StartPos + OtherUsedSegment.Length, Rect.W );
+
+		for ( int32 i = StartFreeSegmentIndex; i < ThisRow.FreeSegments.Num(); ++i )
+		{
+			const FSegment& ThisFreeSegment = ThisRow.FreeSegments[i];
+
+			if ( StartPos >= ThisFreeSegment.StartPos &&
+				EndPos <= ThisFreeSegment.StartPos + ThisFreeSegment.Length )
+			{
+				StartFreeSegmentIndex = i;
+				bFoundSpaceForSegment = true;
+				break;
+			}
+			else if ( !bFoundFutureSpaceForSegment &&
+				StartPos < ThisFreeSegment.StartPos &&
+				OtherUsedSegment.Length <= ThisFreeSegment.Length )
+			{
+				bFoundFutureSpaceForSegment = true;
+				FailedLength = ThisFreeSegment.StartPos - StartPos;
+				break;
+			}
+		}
+
+		if ( !bFoundSpaceForSegment )
+		{
+			if ( !bFoundFutureSpaceForSegment )
+			{
+				FailedLength = Width;
+			}
+
+			return false;
+		}
+	}
+
+	return true;
+}
+
+void FAllocator2D::FlipX( FRect Rect )
+{
+	uint32 MinY = 0;
+
+	uint32 MaxY = Rect.H - 1;
+	for ( int32 Y = Rect.H - 1; Y > 0; --Y )
+	{
+		if ( Rows[ Y ].UsedSegments.Num() > 0 )
+			break;
+
+		--MaxY;
+	}
+
+	for ( uint32 Y = 0; Y <= MaxY; ++Y )
+	{
+		for ( uint32 LowX = 0; LowX < Rect.W / 2; ++LowX )
+		{
+			uint32 HighX = Rect.W - 1 - LowX;
+		
+			const uint64 BitLow = GetBit( LowX, Y );
+			const uint64 BitHigh = GetBit( HighX, Y );
+
+			if ( BitLow != 0ull )
+			{
+				SetBit( HighX, Y );
+			}
+			else
+			{
+				ClearBit( HighX, Y );
+			}
+
+			if ( BitHigh != 0ull )
+			{
+				SetBit( LowX, Y );
+			}
+			else
+			{
+				ClearBit( LowX, Y );
+			}
+		}
+	}
+
+	CreateUsedSegments();
+}
+
+void FAllocator2D::FlipY( FRect Rect )
+{
+	uint32 MinY = 0;
+
+	uint32 MaxY = Rect.H - 1;
+	for ( int32 Y = Rect.H - 1; Y > 0; --Y )
+	{
+		if ( Rows[ Y ].UsedSegments.Num() > 0 )
+			break;
+
+		--MaxY;
+	}
+
+	for ( uint32 LowY = 0; LowY < ( MaxY + 1 ) / 2; ++LowY )
+	{
+		for ( uint32 X = 0; X < Rect.W; ++X )
+		{
+			uint32 HighY = MaxY - LowY;
+
+			const uint64 BitLow = GetBit( X, LowY );
+			const uint64 BitHigh = GetBit( X, HighY );
+
+			if ( BitLow != 0ull )
+			{
+				SetBit( X, HighY );
+			}
+			else
+			{
+				ClearBit( X, HighY );
+			}
+
+			if ( BitHigh != 0ull )
+			{
+				SetBit( X, LowY );
+			}
+			else
+			{
+				ClearBit( X, LowY );
+			}
+		}
+	}
+
+	for ( uint32 LowY = MinY; LowY < ( MaxY + 1 ) / 2; ++LowY )
+	{
+		const int32 HighY = MaxY - LowY;
+
+		Rows.Swap( LowY, HighY );
+		Rows[ LowY ].Index = LowY;
+		Rows[ HighY ].Index = HighY;
+	}
+}
+
+void FAllocator2D::InitSegments()
+{
+	FSegment FreeSegment;
+	FreeSegment.StartPos = 0;
+	FreeSegment.Length = Width;
+
+	for ( FRow& Row : Rows )
+	{
+		Row.FreeSegments.Empty( 1 );
+		Row.FreeSegments.Add( FreeSegment );
+		Row.LongestSegment = FreeSegment.Length;
+
+		Row.UsedSegments.Empty();
+
+	}
+}
+
+void FAllocator2D::CreateUsedSegments()
+{
+	for ( uint32 y = 0; y < Height; ++y )
+	{
+		FRow& CurrentRow = Rows[y];
+		CurrentRow.LongestSegment = 0;
+		CurrentRow.UsedSegments.Empty();
+
+		int32 FirstUsedX = -1;
+		for ( uint32 k = 0; k < Pitch; ++k )
+		{
+			uint32 x = k * 64;
+
+			// If all bits are set
+			if ( Bits[k + y * Pitch] == ~0ull )
+			{
+				if ( FirstUsedX < 0 )
+				{
+					FirstUsedX = x;
+				}
+					
+				if ( k == Pitch - 1 )
+				{
+					AddUsedSegment( CurrentRow, FirstUsedX, x + 64 - FirstUsedX );
+					FirstUsedX = -1;
+				}
+			}
+			// No bits are set
+			else if ( Bits[k + y * Pitch] == 0ull )
+			{
+				if ( FirstUsedX >= 0 )
+				{
+					AddUsedSegment( CurrentRow, FirstUsedX, x - FirstUsedX );
+					FirstUsedX = -1;
+				}
+			}
+			// Some bits are set
+			else
+			{
+				for ( uint32 i = 0; i < 64; ++i )
+				{
+					const uint32 SubX = x + i;
+
+					if ( GetBit(SubX, y) )
+					{
+						if ( FirstUsedX < 0 )
+						{
+							FirstUsedX = SubX;
+						}
+							
+						if ( SubX == Width - 1 )
+						{
+							AddUsedSegment( CurrentRow, FirstUsedX, SubX + 1 - FirstUsedX );
+							FirstUsedX = -1;
+						}
+					}
+					else if ( FirstUsedX >= 0 )
+					{
+						AddUsedSegment( CurrentRow, FirstUsedX, SubX - FirstUsedX );
+						FirstUsedX = -1;
+					}
+				}
+			}
+		}
+	}
+}
+
+void FAllocator2D::AddUsedSegment( FRow& Row, uint32 StartPos, uint32 Length )
+{
+	FSegment UsedSegment;
+	UsedSegment.StartPos = StartPos;
+	UsedSegment.Length = Length;
+
+	Row.UsedSegments.Add(UsedSegment);
+
+	if ( UsedSegment.Length > Row.LongestSegment )
+	{
+		Row.LongestSegment = UsedSegment.Length;
+	}
+}
+
+void FAllocator2D::MergeSegments( FRect Rect, const FAllocator2D& Other )
+{
+	for ( uint32 y = 0; y < Rect.H; ++y )
+	{
+		FRow& ThisRow = Rows[Rect.Y + y];
+		const FRow& OtherRow = Other.Rows[y];
+
+		for ( const FSegment& OtherUsedSegment : OtherRow.UsedSegments )
+		{
+			for ( int32 i = 0; i < ThisRow.FreeSegments.Num(); ++i )
+			{
+				FSegment& ThisFreeSegment = ThisRow.FreeSegments[i];
+
+				uint32 StartPos = Rect.X + OtherUsedSegment.StartPos;
+
+				if ( StartPos >= ThisFreeSegment.StartPos &&
+					StartPos < ThisFreeSegment.StartPos + ThisFreeSegment.Length )
+				{
+					if ( ThisFreeSegment.Length == 1 )
+					{
+						ThisRow.FreeSegments.RemoveAtSwap(i);
+						break;
+					}
+
+					FSegment FirstSegment;
+					FirstSegment.StartPos = ThisFreeSegment.StartPos;
+					FirstSegment.Length = StartPos - ThisFreeSegment.StartPos;
+
+					uint32 EndPos = Rect.X + FMath::Min( OtherUsedSegment.StartPos + OtherUsedSegment.Length, Rect.W ) - 1;
+
+					FSegment SecondSegment;
+					SecondSegment.StartPos = EndPos + 1;
+					SecondSegment.Length = ThisFreeSegment.StartPos + ThisFreeSegment.Length - SecondSegment.StartPos;
+
+					ThisRow.FreeSegments.RemoveAtSwap(i);
+
+					if ( FirstSegment.Length > 0 )
+					{
+						ThisRow.FreeSegments.Add( FirstSegment );
+					}
+					
+					if ( SecondSegment.Length > 0 )
+					{
+						ThisRow.FreeSegments.Add( SecondSegment );
+					}
+
+					break;
+				}
+			}
+		}
+
+		ThisRow.FreeSegments.Sort();
+
+		ThisRow.LongestSegment = 0;
+		for ( const FSegment& Segment : ThisRow.FreeSegments )
+		{
+			ThisRow.LongestSegment = FMath::Max( Segment.Length, ThisRow.LongestSegment );
+		}
+	}
+}
diff --git a/Engine/Source/Developer/MeshUtilitiesCommon/Private/LayoutUV.cpp b/Engine/Source/Developer/MeshUtilitiesCommon/Private/LayoutUV.cpp
new file mode 100644
--- /dev/null
+++ b/Engine/Source/Developer/MeshUtilitiesCommon/Private/LayoutUV.cpp
@@ -0,0 +1,1137 @@
+// Copyright 1998-2018 Epic Games, Inc. All Rights Reserved.
+
+#include "LayoutUV.h"
+#include "DisjointSet.h"
+#include "OverlappingCorners.h"
+#include "Algo/IntroSort.h"
+#include "HAL/PlatformTime.h"
+
+DEFINE_LOG_CATEGORY_STATIC(LogLayoutUV, Warning, All);
+
+#define CHART_JOINING	1
+
+#define NEW_UVS_ARE_SAME THRESH_POINTS_ARE_SAME
+#define LEGACY_UVS_ARE_SAME (1.0f / 1024.0f)
+#define UVLAYOUT_THRESH_UVS_ARE_SAME (GetUVEqualityThreshold())
+
+FLayoutUV::FLayoutUV( IMeshView& InMeshView, uint32 InTextureResolution )
+	: MeshView( InMeshView )
+	, TextureResolution( InTextureResolution )
+	, TotalUVArea( 0.0f )
+	, LayoutRaster( TextureResolution, TextureResolution )
+	, ChartRaster( TextureResolution, TextureResolution )
+	, BestChartRaster( TextureResolution, TextureResolution )
+	, ChartShader( &ChartRaster )
+	, LayoutVersion( ELightmapUVVersion::Latest )
+	, NextMeshChartId( 0 )
+{}
+
+int32 FLayoutUV::FindCharts( const FOverlappingCorners& OverlappingCorners )
+{
+	double Begin = FPlatformTime::Seconds();
+
+	uint32 NumIndexes = MeshView.GetNumIndices();
+	uint32 NumTris = NumIndexes / 3;
+
+	TArray< int32 > TranslatedMatches;
+	TranslatedMatches.SetNumUninitialized( NumIndexes );
+	TexCoords.SetNumUninitialized( NumIndexes );
+	for( uint32 i = 0; i < NumIndexes; i++ )
+	{
+		TranslatedMatches[i] = -1;
+		TexCoords[i] = MeshView.GetInputTexcoord(i);
+	}
+
+	// Build disjoint set
+	FDisjointSet DisjointSet( NumTris );
+	for( uint32 i = 0; i < NumIndexes; i++ )
+	{
+		const TArray<int32>& Overlapping = OverlappingCorners.FindIfOverlapping(i);
+		for (int32 It : Overlapping)
+		{
+			uint32 j = It;
+			if( j > i )
+			{
+				const uint32 TriI = i/3;
+				const uint32 TriJ = j/3;
+
+				bool bUnion = false;
+
+#if CHART_JOINING
+				bool bPositionMatch = PositionsMatch( i, j );
+				if( bPositionMatch )
+				{
+					uint32 i1 = 3 * TriI + (i + 1) % 3;
+					uint32 i2 = 3 * TriI + (i + 2) % 3;
+					uint32 j1 = 3 * TriJ + (j + 1) % 3;
+					uint32 j2 = 3 * TriJ + (j + 2) % 3;
+
+					bool bEdgeMatch21 = PositionsMatch( i2, j1 );
+					bool bEdgeMatch12 = PositionsMatch( i1, j2 );
+					if( bEdgeMatch21 || bEdgeMatch12 )
+					{
+						uint32 ie = bEdgeMatch21 ? i2 : i1;
+						uint32 je = bEdgeMatch21 ? j1 : j2;
+							
+						bool bUVMatch = UVsMatch( i, j ) && UVsMatch( ie, je );
+						bool bUVWindingMatch = TriangleUVArea( TriI ) * TriangleUVArea( TriJ ) >= 0.0f;
+						if( bUVMatch && bUVWindingMatch )
+						{
+							bUnion = true;
+						}
+						else if( NormalsMatch( i, j ) && NormalsMatch( ie, je ) )
+						{
+							// Chart edge
+							FVector2D EdgeUVi = TexCoords[ie] - TexCoords[i];
+							FVector2D EdgeUVj = TexCoords[je] - TexCoords[j];
+							
+							// Would these edges match if the charts were translated
+							bool bTranslatedUVMatch = ( EdgeUVi - EdgeUVj ).IsNearlyZero(UVLAYOUT_THRESH_UVS_ARE_SAME);
+							if( bTranslatedUVMatch )
+							{
+								// Note: may be mirrored
+								
+								// TODO should these be restricted to axis aligned edges?
+								uint32 EdgeI = bEdgeMatch21 ? i2 : i;
+								uint32 EdgeJ = bEdgeMatch21 ? j : j2;
+
+								// Only allow one match per edge
+								if( TranslatedMatches[ EdgeI ] < 0 &&
+									TranslatedMatches[ EdgeJ ] < 0 )
+								{
+									TranslatedMatches[ EdgeI ] = EdgeJ;
+									TranslatedMatches[ EdgeJ ] = EdgeI;
+								}
+							}
+						}
+					}
+				}
+#else
+				if( VertsMatch( i, j ) )
+				{
+					// Edge must match as well (same winding)
+					if( VertsMatch( 3 * TriI + (i - 1) % 3, 3 * TriJ + (j + 1) % 3 ) ||
+						VertsMatch( 3 * TriI + (i + 1) % 3, 3 * TriJ + (j - 1) % 3 ) )
+					{
+						// Check for UV winding match too
+						if( TriangleUVArea( TriI ) * TriangleUVArea( TriJ ) >= 0.0f )
+						{
+							bUnion = true;
+						}
+					}
+				}
+#endif
+
+				if( bUnion )
+				{
+					// TODO solve spiral case by checking sets for UV overlap
+					DisjointSet.Union( TriI, TriJ );
+				}
+			}
+		}
+	}
+
+	// Sort tris by chart
+	SortedTris.SetNumUninitialized( NumTris );
+	for( uint32 i = 0; i < NumTris; i++ )
+	{
+		// Flatten disjoint set path
+		DisjointSet.Find(i);
+		SortedTris[i] = i;
+	}
+
+	struct FCompareTris
+	{
+		FDisjointSet* DisjointSet;
+
+		FCompareTris( FDisjointSet* InDisjointSet )
+		: DisjointSet( InDisjointSet )
+		{}
+
+		FORCEINLINE bool operator()( uint32 A, uint32 B ) const
+		{
+			return (*DisjointSet)[A] < (*DisjointSet)[B];
+		}
+	};
+
+	Algo::IntroSort( SortedTris, FCompareTris( &DisjointSet ) );
+
+	TMap< uint32, int32 > DisjointSetToChartMap;
+
+	// Build Charts
+	for( uint32 Tri = 0; Tri < NumTris; )
+	{
+		int32 i = Charts.AddUninitialized();
+		FMeshChart& Chart = Charts[i];
+		Chart.Id = NextMeshChartId++;
+
+		Chart.MinUV = FVector2D( FLT_MAX, FLT_MAX );
+		Chart.MaxUV = FVector2D( -FLT_MAX, -FLT_MAX );
+		Chart.UVArea = 0.0f;
+		Chart.WorldScale = FVector2D::ZeroVector;
+		FMemory::Memset( Chart.Join, 0xff );
+
+		Chart.FirstTri = Tri;
+
+		uint32 ChartID = DisjointSet[ SortedTris[ Tri ] ];
+		DisjointSetToChartMap.Add( ChartID, i );
+
+		for( ; Tri < NumTris && DisjointSet[ SortedTris[ Tri ] ] == ChartID; Tri++ )
+		{
+			// Calculate chart bounds
+			FVector		Positions[3];
+			FVector2D	UVs[3];
+			for( int k = 0; k < 3; k++ )
+			{
+				uint32 Index = 3 * SortedTris[ Tri ] + k;
+
+				Positions[k] = MeshView.GetPosition( Index );
+				UVs[k] = TexCoords[ Index ];
+
+				Chart.MinUV.X = FMath::Min( Chart.MinUV.X, UVs[k].X );
+				Chart.MinUV.Y = FMath::Min( Chart.MinUV.Y, UVs[k].Y );
+				Chart.MaxUV.X = FMath::Max( Chart.MaxUV.X, UVs[k].X );
+				Chart.MaxUV.Y = FMath::Max( Chart.MaxUV.Y, UVs[k].Y );
+			}
+
+			FVector Edge1 = Positions[1] - Positions[0];
+			FVector Edge2 = Positions[2] - Positions[0];
+			float Area = 0.5f * ( Edge1 ^ Edge2 ).Size();
+
+			FVector2D EdgeUV1 = UVs[1] - UVs[0];
+			FVector2D EdgeUV2 = UVs[2] - UVs[0];
+			float UVArea = 0.5f * FMath::Abs( EdgeUV1.X * EdgeUV2.Y - EdgeUV1.Y * EdgeUV2.X );
+
+			FVector2D UVLength;
+			UVLength.X = ( EdgeUV2.Y * Edge1 - EdgeUV1.Y * Edge2 ).Size();
+			UVLength.Y = (-EdgeUV2.X * Edge1 + EdgeUV1.X * Edge2 ).Size();
+			
+			Chart.WorldScale += UVLength;
+			Chart.UVArea += UVArea;
+		}
+		
+		Chart.LastTri = Tri;
+
+#if !CHART_JOINING
+		if (LayoutVersion >= ELightmapUVVersion::SmallChartPacking)
+		{
+			Chart.WorldScale /= FMath::Max(Chart.UVArea, 1e-8f);
+		}
+		else
+		{
+			if (Chart.UVArea > 1e-4f)
+			{
+				Chart.WorldScale /= Chart.UVArea;
+			}
+			else
+			{
+				Chart.WorldScale = FVector2D::ZeroVector;
+			}
+		}		
+
+		TotalUVArea += Chart.UVArea * Chart.WorldScale.X * Chart.WorldScale.Y;
+#endif
+	}
+
+#if CHART_JOINING
+	for( int32 i = 0; i < Charts.Num(); i++ )
+	{
+		FMeshChart& Chart = Charts[i];
+
+		for( uint32 Tri = Chart.FirstTri; Tri < Chart.LastTri; Tri++ )
+		{
+			for( int k = 0; k < 3; k++ )
+			{
+				uint32 Index = 3 * SortedTris[ Tri ] + k;
+
+				if( TranslatedMatches[ Index ] >= 0 )
+				{
+					checkSlow( TranslatedMatches[ TranslatedMatches[ Index ] ] == Index );
+
+					uint32 V0i = Index;
+					uint32 V0j = TranslatedMatches[ Index ];
+
+					uint32 TriI = V0i / 3;
+					uint32 TriJ = V0j / 3;
+					
+					if( TriJ <= TriI )
+					{
+						// Only need to consider one direction
+						continue;
+					}
+					
+					uint32 V1i = 3 * TriI + (V0i + 1) % 3;
+					uint32 V1j = 3 * TriJ + (V0j + 1) % 3;
+
+					int32 ChartI = i;
+					int32 ChartJ = DisjointSetToChartMap[ DisjointSet[ TriJ ] ];
+
+					FVector2D UV0i = TexCoords[ V0i ];
+					FVector2D UV1i = TexCoords[ V1i ];
+					FVector2D UV0j = TexCoords[ V0j ];
+					FVector2D UV1j = TexCoords[ V1j ];
+
+					FVector2D EdgeUVi = UV1i - UV0i;
+					FVector2D EdgeUVj = UV1j - UV0j;
+
+					bool bMirrored = TriangleUVArea( TriI ) * TriangleUVArea( TriJ ) < 0.0f;
+					
+					FVector2D EdgeOffset0 = UV0i - UV1j;
+					FVector2D EdgeOffset1 = UV1i - UV0j;
+					checkSlow( ( EdgeOffset0 - EdgeOffset1 ).IsNearlyZero(UVLAYOUT_THRESH_UVS_ARE_SAME) );
+
+					FVector2D Translation = EdgeOffset0;
+
+					FMeshChart& ChartA = Charts[ ChartI ];
+					FMeshChart& ChartB = Charts[ ChartJ ];
+
+					for( uint32 Side = 0; Side < 4; Side++ )
+					{
+						// Join[] = { left, right, bottom, top }
+
+						// FIXME
+						if( bMirrored )
+							continue;
+
+						if( ChartA.Join[ Side ^ 0 ] != -1 ||
+							ChartB.Join[ Side ^ 1 ] != -1 )
+						{
+							// Already joined with something else
+							continue;
+						}
+
+						uint32 Sign = Side & 1;
+						uint32 Axis = Side >> 1;
+
+						bool bAxisAligned = FMath::Abs( EdgeUVi[ Axis ] ) < UVLAYOUT_THRESH_UVS_ARE_SAME;
+						bool bBorderA = FMath::Abs( UV0i[ Axis ] - ( Sign ^ 0 ? Chart.MaxUV[ Axis ] : Chart.MinUV[ Axis ] ) ) < UVLAYOUT_THRESH_UVS_ARE_SAME;
+						bool bBorderB = FMath::Abs( UV0j[ Axis ] - ( Sign ^ 1 ? Chart.MaxUV[ Axis ] : Chart.MinUV[ Axis ] ) ) < UVLAYOUT_THRESH_UVS_ARE_SAME;
+
+						// FIXME mirrored
+						if( !bAxisAligned || !bBorderA || !bBorderB )
+						{
+							// Edges weren't on matching rectangle borders
+							continue;
+						}
+
+						FVector2D CenterA = 0.5f * ( ChartA.MinUV + ChartA.MaxUV );
+						FVector2D CenterB = 0.5f * ( ChartB.MinUV + ChartB.MaxUV );
+
+						FVector2D ExtentA = 0.5f * ( ChartA.MaxUV - ChartA.MinUV );
+						FVector2D ExtentB = 0.5f * ( ChartB.MaxUV - ChartB.MinUV );
+
+						// FIXME mirrored
+						CenterB += Translation;
+
+						FVector2D CenterDiff = CenterA - CenterB;
+						FVector2D ExtentDiff = ExtentA - ExtentB;
+						FVector2D Separation = ExtentA + ExtentB + CenterDiff * ( Sign ? 1.0f : -1.0f );
+
+						bool bCenterMatch = FMath::Abs( CenterDiff[ Axis ^ 1 ] ) < UVLAYOUT_THRESH_UVS_ARE_SAME;
+						bool bExtentMatch = FMath::Abs( ExtentDiff[ Axis ^ 1 ] ) < UVLAYOUT_THRESH_UVS_ARE_SAME;
+						bool bSeparate    = FMath::Abs( Separation[ Axis ^ 0 ] ) < UVLAYOUT_THRESH_UVS_ARE_SAME;
+	
+						if( !bCenterMatch || !bExtentMatch || !bSeparate )
+						{
+							// Rectangles don't match up after translation
+							continue;
+						}
+
+						// Found a valid edge join
+						ChartA.Join[ Side ^ 0 ] = ChartJ;
+						ChartB.Join[ Side ^ 1 ] = ChartI;
+						break;
+					}
+				}
+			}
+		}
+	}
+
+	TArray< uint32 > JoinedSortedTris;
+	JoinedSortedTris.Reserve( NumTris );
+
+	// Detect loops
+	for( uint32 Axis = 0; Axis < 2; Axis++ )
+	{
+		uint32 Side = Axis << 1;
+
+		for( int32 i = 0; i < Charts.Num(); i++ )
+		{
+			int32 j = Charts[i].Join[ Side ^ 1 ];
+			while( j != -1 )
+			{
+				int32 Next = Charts[j].Join[ Side ^ 1 ];
+				if( Next == i )
+				{
+					// Break loop
+					Charts[i].Join[ Side ^ 0 ] = -1;
+					Charts[j].Join[ Side ^ 1 ] = -1;
+					break;
+				}
+				j = Next;
+			}
+		}
+	}
+
+	// Join rows first, then columns
+	for( uint32 Axis = 0; Axis < 2; Axis++ )
+	{
+		for( int32 i = 0; i < Charts.Num(); i++ )
+		{
+			FMeshChart& ChartA = Charts[i];
+
+			if( ChartA.FirstTri == ChartA.LastTri )
+			{
+				// Empty chart
+				continue;
+			}
+			
+			for( uint32 Side = 0; Side < 4; Side++ )
+			{
+				if( ChartA.Join[ Side ] != -1 )
+				{
+					FMeshChart& ChartB = Charts[ ChartA.Join[ Side ] ];
+
+					check( ChartB.Join[ Side ^ 1 ] == i );
+					check( ChartB.FirstTri != ChartB.LastTri );
+				}
+			}
+		}
+
+		NumTris = 0;
+
+		for( int32 i = 0; i < Charts.Num(); i++ )
+		{
+			FMeshChart& Chart = Charts[i];
+			NumTris += Chart.LastTri - Chart.FirstTri;
+		}
+		check( NumTris == SortedTris.Num() );
+
+		NumTris = 0;
+		for( int32 i = 0; i < Charts.Num(); i++ )
+		{
+			FMeshChart& ChartA = Charts[i];
+
+			if( ChartA.FirstTri == ChartA.LastTri )
+			{
+				// Empty chart
+				continue;
+			}
+
+			uint32 Side = Axis << 1;
+
+			// Find start (left, bottom)
+			if( ChartA.Join[ Side ^ 0 ] == -1 )
+			{
+				// Add original tris
+				NumTris += ChartA.LastTri - ChartA.FirstTri;
+
+				// Continue joining until no more to the (right, top)
+				int32 Next = ChartA.Join[ Side ^ 1 ];
+				while( Next != -1 )
+				{
+					FMeshChart& ChartB = Charts[ Next ];
+
+					NumTris += ChartB.LastTri - ChartB.FirstTri;
+					Next = ChartB.Join[ Side ^ 1 ];
+				}
+			}
+		}
+		check( NumTris == SortedTris.Num() );
+
+#if 1
+		NumTris = 0;
+		for( int32 i = 0; i < Charts.Num(); i++ )
+		{
+			FMeshChart& ChartA = Charts[i];
+
+			if( ChartA.FirstTri == ChartA.LastTri )
+			{
+				// Empty chart
+				continue;
+			}
+
+			// Join[] = { left, right, bottom, top }
+
+			uint32 Side = Axis << 1;
+
+			// Find start (left, bottom)
+			if( ChartA.Join[ Side ^ 0 ] == -1 )
+			{
+				uint32 FirstTri = JoinedSortedTris.Num();
+			
+				// Add original tris
+				for( uint32 Tri = ChartA.FirstTri; Tri < ChartA.LastTri; Tri++ )
+				{
+					JoinedSortedTris.Add( SortedTris[ Tri ] );
+				}
+				NumTris += ChartA.LastTri - ChartA.FirstTri;
+
+				// Continue joining until no more to the (right, top)
+				while( ChartA.Join[ Side ^ 1 ] != -1 )
+				{
+					FMeshChart& ChartB = Charts[ ChartA.Join[ Side ^ 1 ] ];
+
+					check( ChartB.FirstTri != ChartB.LastTri );
+
+					FVector2D Translation = ChartA.MinUV - ChartB.MinUV;
+					Translation[ Axis ] += ChartA.MaxUV[ Axis ] - ChartA.MinUV[ Axis ];
+
+					for( uint32 Tri = ChartB.FirstTri; Tri < ChartB.LastTri; Tri++ )
+					{
+						JoinedSortedTris.Add( SortedTris[ Tri ] );
+						for( int k = 0; k < 3; k++ )
+						{
+							TexCoords[ 3 * SortedTris[ Tri ] + k ] += Translation;
+						}
+					}
+					NumTris += ChartB.LastTri - ChartB.FirstTri;
+
+					ChartA.Join[ Side ^ 1 ] = ChartB.Join[ Side ^ 1 ];
+					ChartA.MaxUV[ Axis ] += ChartB.MaxUV[ Axis ] - ChartB.MinUV[ Axis ];
+					if( LayoutVersion >= ELightmapUVVersion::ChartJoiningLFix )
+					{
+						// Fixing joined chart MaxUV value to properly inflate non-joined axis extent
+						ChartA.MaxUV[ Axis ^ 1 ] = FMath::Max( ChartA.MaxUV[ Axis ^ 1 ], ChartA.MinUV[ Axis ^ 1 ] + ( ChartB.MaxUV[ Axis ^ 1 ] - ChartB.MinUV[ Axis ^ 1 ] ) );
+					}
+					ChartA.WorldScale += ChartB.WorldScale;
+					ChartA.UVArea += ChartB.UVArea;
+
+					ChartB.FirstTri = 0;
+					ChartB.LastTri = 0;
+					ChartB.UVArea = 0.0f;
+
+					DisconnectChart( ChartB, Side ^ 2 );
+					DisconnectChart( ChartB, Side ^ 3 );
+				}
+
+				ChartA.FirstTri = FirstTri;
+				ChartA.LastTri = JoinedSortedTris.Num();
+			}
+			else
+			{
+				// Make sure a starting chart could connect to this
+				FMeshChart& ChartB = Charts[ ChartA.Join[ Side ^ 0 ] ];
+				check( ChartB.Join[ Side ^ 1 ] == i );
+				check( ChartB.FirstTri != ChartB.LastTri );
+			}
+		}
+		check( NumTris == SortedTris.Num() );
+
+		check( SortedTris.Num() == JoinedSortedTris.Num() );
+		Exchange( SortedTris, JoinedSortedTris );
+		JoinedSortedTris.Reset();
+#endif
+	}
+	
+	// Clean out empty charts
+	for( int32 i = 0; i < Charts.Num(); i++ )
+	{
+		while( i < Charts.Num() && Charts[i].FirstTri == Charts[i].LastTri )
+		{
+			Charts.RemoveAtSwap(i);
+		}
+	}
+
+	for( int32 i = 0; i < Charts.Num(); i++ )
+	{
+		FMeshChart& Chart = Charts[i];
+
+		if (LayoutVersion >= ELightmapUVVersion::SmallChartPacking)
+		{
+			Chart.WorldScale /= FMath::Max(Chart.UVArea, 1e-8f);
+		}
+		else
+		{
+			if (Chart.UVArea > 1e-4f)
+			{
+				Chart.WorldScale /= Chart.UVArea;
+			}
+			else
+			{
+				Chart.WorldScale = FVector2D::ZeroVector;
+			}
+		}
+
+		TotalUVArea += Chart.UVArea * Chart.WorldScale.X * Chart.WorldScale.Y;
+	}
+#endif
+
+	double End = FPlatformTime::Seconds();
+
+	UE_LOG(LogLayoutUV, Display, TEXT("FindCharts: %s"), *FPlatformTime::PrettyTime(End - Begin) );
+
+	return Charts.Num();
+}
+
+bool FLayoutUV::FindBestPacking()
+{
+	if( (uint32)Charts.Num() > TextureResolution * TextureResolution || TotalUVArea == 0.f )
+	{
+		// More charts than texels
+		return false;
+	}
+	
+	const float LinearSearchStart = 0.5f;
+	const float LinearSearchStep = 0.5f;
+	const int32 BinarySearchSteps = 6;
+	
+	float UVScaleFail = TextureResolution * FMath::Sqrt( 1.0f / TotalUVArea );
+	float UVScalePass = TextureResolution * FMath::Sqrt( LinearSearchStart / TotalUVArea );
+
+	// Linear search for first fit
+	while(1)
+	{
+		ScaleCharts( UVScalePass );
+
+		bool bFit = PackCharts();
+		if( bFit )
+		{
+			break;
+		}
+
+		UVScaleFail = UVScalePass;
+		UVScalePass *= LinearSearchStep;
+	}
+
+	// Binary search for best fit
+	for( int32 i = 0; i < BinarySearchSteps; i++ )
+	{
+		float UVScale = 0.5f * ( UVScaleFail + UVScalePass );
+		ScaleCharts( UVScale );
+
+		bool bFit = PackCharts();
+		if( bFit )
+		{
+			UVScalePass = UVScale;
+		}
+		else
+		{
+			UVScaleFail = UVScale;
+		}
+	}
+
+	// TODO store packing scale/bias separate so this isn't necessary
+	ScaleCharts( UVScalePass );
+	PackCharts();
+
+	return true;
+}
+
+void FLayoutUV::ScaleCharts( float UVScale )
+{
+	for( int32 i = 0; i < Charts.Num(); i++ )
+	{
+		FMeshChart& Chart = Charts[i];
+		Chart.UVScale = Chart.WorldScale * UVScale;
+	}
+	
+	if ( LayoutVersion >= ELightmapUVVersion::ScaleChartsOrderingFix )
+	{
+		// Unsort the charts to make sure ScaleCharts always return the same ordering
+		Algo::IntroSort( Charts, []( const FMeshChart& A, const FMeshChart& B )
+		{
+			return A.Id < B.Id;
+		});
+	}
+
+	// Scale charts such that they all fit and roughly total the same area as before
+#if 1
+	float UniformScale = 1.0f;
+	for( int i = 0; i < 1000; i++ )
+	{
+		uint32 NumMaxedOut = 0;
+		float ScaledUVArea = 0.0f;
+		for( int32 ChartIndex = 0; ChartIndex < Charts.Num(); ChartIndex++ )
+		{
+			FMeshChart& Chart = Charts[ChartIndex];
+
+			FVector2D ChartSize	= Chart.MaxUV - Chart.MinUV;
+			FVector2D ChartSizeScaled = ChartSize * Chart.UVScale * UniformScale;
+
+			const float MaxChartEdge = TextureResolution - 1.0f;
+			const float LongestChartEdge = FMath::Max( ChartSizeScaled.X, ChartSizeScaled.Y );
+
+			const float Epsilon = 0.01f;
+			if( LongestChartEdge + Epsilon > MaxChartEdge )
+			{
+				// Rescale oversized charts to fit
+				Chart.UVScale.X = MaxChartEdge / FMath::Max( ChartSize.X, ChartSize.Y );
+				Chart.UVScale.Y = MaxChartEdge / FMath::Max( ChartSize.X, ChartSize.Y );
+				NumMaxedOut++;
+			}
+			else
+			{
+				Chart.UVScale.X *= UniformScale;
+				Chart.UVScale.Y *= UniformScale;
+			}
+			
+			ScaledUVArea += Chart.UVArea * Chart.UVScale.X * Chart.UVScale.Y;
+		}
+
+		if( NumMaxedOut == 0 )
+		{
+			// No charts maxed out so no need to rebalance
+			break;
+		}
+
+		if( NumMaxedOut == Charts.Num() )
+		{
+			// All charts are maxed out
+			break;
+		}
+
+		// Scale up smaller charts to maintain expected total area
+		// Want ScaledUVArea == TotalUVArea * UVScale^2
+		float RebalanceScale = UVScale * FMath::Sqrt( TotalUVArea / ScaledUVArea );
+		if( RebalanceScale < 1.01f )
+		{
+			// Stop if further rebalancing is minor
+			break;
+		}
+		UniformScale = RebalanceScale;
+	}
+#endif
+
+#if 1
+	float NonuniformScale = 1.0f;
+	for( int i = 0; i < 1000; i++ )
+	{
+		uint32 NumMaxedOut = 0;
+		float ScaledUVArea = 0.0f;
+		for( int32 ChartIndex = 0; ChartIndex < Charts.Num(); ChartIndex++ )
+		{
+			FMeshChart& Chart = Charts[ChartIndex];
+
+			for( int k = 0; k < 2; k++ )
+			{
+				const float MaximumChartSize = TextureResolution - 1.0f;
+				const float ChartSize = Chart.MaxUV[k] - Chart.MinUV[k];
+				const float ChartSizeScaled = ChartSize * Chart.UVScale[k] * NonuniformScale;
+
+				const float Epsilon = 0.01f;
+				if( ChartSizeScaled + Epsilon > MaximumChartSize )
+				{
+					// Scale oversized charts to max size
+					Chart.UVScale[k] = MaximumChartSize / ChartSize;
+					NumMaxedOut++;
+				}
+				else
+				{
+					Chart.UVScale[k] *= NonuniformScale;
+				}
+			}
+
+			ScaledUVArea += Chart.UVArea * Chart.UVScale.X * Chart.UVScale.Y;
+		}
+
+		if( NumMaxedOut == 0 )
+		{
+			// No charts maxed out so no need to rebalance
+			break;
+		}
+
+		if( NumMaxedOut == Charts.Num() * 2 )
+		{
+			// All charts are maxed out in both dimensions
+			break;
+		}
+
+		// Scale up smaller charts to maintain expected total area
+		// Want ScaledUVArea == TotalUVArea * UVScale^2
+		float RebalanceScale = UVScale * FMath::Sqrt( TotalUVArea / ScaledUVArea );
+		if( RebalanceScale < 1.01f )
+		{
+			// Stop if further rebalancing is minor
+			break;
+		}
+		NonuniformScale = RebalanceScale;
+	}
+#endif
+
+	// Sort charts from largest to smallest
+	struct FCompareCharts
+	{
+		FORCEINLINE bool operator()( const FMeshChart& A, const FMeshChart& B ) const
+		{
+			// Rect area
+			FVector2D ChartRectA = ( A.MaxUV - A.MinUV ) * A.UVScale;
+			FVector2D ChartRectB = ( B.MaxUV - B.MinUV ) * B.UVScale;
+			return ChartRectA.X * ChartRectA.Y > ChartRectB.X * ChartRectB.Y;
+		}
+	};
+	Algo::IntroSort( Charts, FCompareCharts() );
+}
+
+bool FLayoutUV::PackCharts()
+{
+	uint32 RasterizeCycles = 0;
+	uint32 FindCycles = 0;
+	
+	double BeginPackCharts = FPlatformTime::Seconds();
+
+	LayoutRaster.Clear();
+
+	for( int32 i = 0; i < Charts.Num(); i++ )
+	{
+		FMeshChart& Chart = Charts[i];
+
+		// Try different orientations and pick best
+		int32				BestOrientation = -1;
+		FAllocator2D::FRect	BestRect = { ~0u, ~0u, ~0u, ~0u };
+
+		for( int32 Orientation = 0; Orientation < 8; Orientation++ )
+		{
+			// TODO If any dimension is less than 1 pixel shrink dimension to zero
+
+			OrientChart( Chart, Orientation );
+			
+			FVector2D ChartSize = Chart.MaxUV - Chart.MinUV;
+			ChartSize = ChartSize.X * Chart.PackingScaleU + ChartSize.Y * Chart.PackingScaleV;
+
+			// Only need half pixel dilate for rects
+			FAllocator2D::FRect	Rect;
+			Rect.X = 0;
+			Rect.Y = 0;
+			Rect.W = FMath::CeilToInt( FMath::Abs( ChartSize.X ) + 1.0f );
+			Rect.H = FMath::CeilToInt( FMath::Abs( ChartSize.Y ) + 1.0f );
+
+			// Just in case lack of precision pushes it over
+			Rect.W = FMath::Min( TextureResolution, Rect.W );
+			Rect.H = FMath::Min( TextureResolution, Rect.H );
+
+			const bool bRectPack = false;
+
+			if( bRectPack )
+			{
+				if( LayoutRaster.Find( Rect ) )
+				{
+					// Is best?
+					if( Rect.X + Rect.Y * TextureResolution < BestRect.X + BestRect.Y * TextureResolution )
+					{
+						BestOrientation = Orientation;
+						BestRect = Rect;
+					}
+				}
+				else
+				{
+					continue;
+				}
+			}
+			else
+			{
+				if ( LayoutVersion >= ELightmapUVVersion::Segments && Orientation % 4 == 1 )
+				{
+					ChartRaster.FlipX( Rect );
+				}
+				else if ( LayoutVersion >= ELightmapUVVersion::Segments && Orientation % 4 == 3 )
+				{
+					ChartRaster.FlipY( Rect );
+				}
+				else
+				{
+					int32 BeginRasterize = FPlatformTime::Cycles();
+					RasterizeChart( Chart, Rect.W, Rect.H );
+					RasterizeCycles += FPlatformTime::Cycles() - BeginRasterize;
+				}
+
+				bool bFound = false;
+
+				uint32 BeginFind = FPlatformTime::Cycles();
+				if ( LayoutVersion == ELightmapUVVersion::BitByBit )
+				{
+					bFound = LayoutRaster.FindBitByBit( Rect, ChartRaster );
+				}
+				else if ( LayoutVersion >= ELightmapUVVersion::Segments )
+				{
+					bFound = LayoutRaster.FindWithSegments( Rect, BestRect, ChartRaster );
+				}
+				FindCycles += FPlatformTime::Cycles() - BeginFind;
+
+				if( bFound )
+				{
+					// Is best?
+					if( Rect.X + Rect.Y * TextureResolution < BestRect.X + BestRect.Y * TextureResolution )
+					{
+						BestChartRaster = ChartRaster;
+
+						BestOrientation = Orientation;
+						BestRect = Rect;
+
+						if ( BestRect.X == 0 && BestRect.Y == 0 )
+						{
+							// BestRect can't be beat, stop here
+							break;
+						}
+					}
+				}
+				else
+				{
+					continue;
+				}
+			}
+		}
+
+		if( BestOrientation >= 0 )
+		{
+			// Add chart to layout
+			OrientChart( Chart, BestOrientation );
+
+			LayoutRaster.Alloc( BestRect, BestChartRaster );
+
+			Chart.PackingBias.X += BestRect.X;
+			Chart.PackingBias.Y += BestRect.Y;
+		}
+		else
+		{
+			// Found no orientation that fit
+			return false;
+		}
+	}
+
+	double EndPackCharts = FPlatformTime::Seconds();
+
+	UE_LOG(LogLayoutUV, Display, TEXT("PackCharts: %s"),	*FPlatformTime::PrettyTime(EndPackCharts - BeginPackCharts));
+	UE_LOG(LogLayoutUV, Display, TEXT("  Rasterize: %u"),	RasterizeCycles);
+	UE_LOG(LogLayoutUV, Display, TEXT("  Find: %u"),		FindCycles);
+
+	return true;
+}
+
+void FLayoutUV::OrientChart( FMeshChart& Chart, int32 Orientation )
+{
+	switch( Orientation )
+	{
+	case 0:
+		// 0 degrees
+		Chart.PackingScaleU = FVector2D( Chart.UVScale.X, 0 );
+		Chart.PackingScaleV = FVector2D( 0, Chart.UVScale.Y );
+		Chart.PackingBias = -Chart.MinUV.X * Chart.PackingScaleU - Chart.MinUV.Y * Chart.PackingScaleV + 0.5f;
+		break;
+	case 1:
+		// 0 degrees, flip x
+		Chart.PackingScaleU = FVector2D( -Chart.UVScale.X, 0 );
+		Chart.PackingScaleV = FVector2D( 0, Chart.UVScale.Y );
+		Chart.PackingBias = -Chart.MaxUV.X * Chart.PackingScaleU - Chart.MinUV.Y * Chart.PackingScaleV + 0.5f;
+		break;
+	case 2:
+		// 90 degrees
+		Chart.PackingScaleU = FVector2D( 0, -Chart.UVScale.X );
+		Chart.PackingScaleV = FVector2D( Chart.UVScale.Y, 0 );
+		Chart.PackingBias = -Chart.MaxUV.X * Chart.PackingScaleU - Chart.MinUV.Y * Chart.PackingScaleV + 0.5f;
+		break;
+	case 3:
+		// 90 degrees, flip x
+		Chart.PackingScaleU = FVector2D( 0, Chart.UVScale.X );
+		Chart.PackingScaleV = FVector2D( Chart.UVScale.Y, 0 );
+		Chart.PackingBias = -Chart.MinUV.X * Chart.PackingScaleU - Chart.MinUV.Y * Chart.PackingScaleV + 0.5f;
+		break;
+	case 4:
+		// 180 degrees
+		Chart.PackingScaleU = FVector2D( -Chart.UVScale.X, 0 );
+		Chart.PackingScaleV = FVector2D( 0, -Chart.UVScale.Y );
+		Chart.PackingBias = -Chart.MaxUV.X * Chart.PackingScaleU - Chart.MaxUV.Y * Chart.PackingScaleV + 0.5f;
+		break;
+	case 5:
+		// 180 degrees, flip x
+		Chart.PackingScaleU = FVector2D( Chart.UVScale.X, 0 );
+		Chart.PackingScaleV = FVector2D( 0, -Chart.UVScale.Y );
+		Chart.PackingBias = -Chart.MinUV.X * Chart.PackingScaleU - Chart.MaxUV.Y * Chart.PackingScaleV + 0.5f;
+		break;
+	case 6:
+		// 270 degrees
+		Chart.PackingScaleU = FVector2D( 0, Chart.UVScale.X );
+		Chart.PackingScaleV = FVector2D( -Chart.UVScale.Y, 0 );
+		Chart.PackingBias = -Chart.MinUV.X * Chart.PackingScaleU - Chart.MaxUV.Y * Chart.PackingScaleV + 0.5f;
+		break;
+	case 7:
+		// 270 degrees, flip x
+		Chart.PackingScaleU = FVector2D( 0, -Chart.UVScale.X );
+		Chart.PackingScaleV = FVector2D( -Chart.UVScale.Y, 0 );
+		Chart.PackingBias = -Chart.MaxUV.X * Chart.PackingScaleU - Chart.MaxUV.Y * Chart.PackingScaleV + 0.5f;
+		break;
+	}
+}
+
+// Max of 2048x2048 due to precision
+// Dilate in 28.4 fixed point. Half pixel dilation is conservative rasterization.
+// Dilation same as Minkowski sum of triangle and square.
+template< typename TShader, int32 Dilate >
+void RasterizeTriangle( TShader& Shader, const FVector2D Points[3], int32 ScissorWidth, int32 ScissorHeight )
+{
+	const FVector2D HalfPixel( 0.5f, 0.5f );
+	FVector2D p0 = Points[0] - HalfPixel;
+	FVector2D p1 = Points[1] - HalfPixel;
+	FVector2D p2 = Points[2] - HalfPixel;
+
+	// Correct winding
+	float Facing = ( p0.X - p1.X ) * ( p2.Y - p0.Y ) - ( p0.Y - p1.Y ) * ( p2.X - p0.X );
+	if( Facing < 0.0f )
+	{
+		Swap( p0, p2 );
+	}
+
+	// 28.4 fixed point
+	const int32 X0 = (int32)( 16.0f * p0.X + 0.5f );
+	const int32 X1 = (int32)( 16.0f * p1.X + 0.5f );
+	const int32 X2 = (int32)( 16.0f * p2.X + 0.5f );
+	
+	const int32 Y0 = (int32)( 16.0f * p0.Y + 0.5f );
+	const int32 Y1 = (int32)( 16.0f * p1.Y + 0.5f );
+	const int32 Y2 = (int32)( 16.0f * p2.Y + 0.5f );
+
+	// Bounding rect
+	int32 MinX = ( FMath::Min3( X0, X1, X2 ) - Dilate + 15 ) / 16;
+	int32 MaxX = ( FMath::Max3( X0, X1, X2 ) + Dilate + 15 ) / 16;
+	int32 MinY = ( FMath::Min3( Y0, Y1, Y2 ) - Dilate + 15 ) / 16;
+	int32 MaxY = ( FMath::Max3( Y0, Y1, Y2 ) + Dilate + 15 ) / 16;
+
+	// Clip to image
+	MinX = FMath::Clamp( MinX, 0, ScissorWidth );
+	MaxX = FMath::Clamp( MaxX, 0, ScissorWidth );
+	MinY = FMath::Clamp( MinY, 0, ScissorHeight );
+	MaxY = FMath::Clamp( MaxY, 0, ScissorHeight );
+
+	// Deltas
+	const int32 DX01 = X0 - X1;
+	const int32 DX12 = X1 - X2;
+	const int32 DX20 = X2 - X0;
+
+	const int32 DY01 = Y0 - Y1;
+	const int32 DY12 = Y1 - Y2;
+	const int32 DY20 = Y2 - Y0;
+
+	// Half-edge constants
+	int32 C0 = DY01 * X0 - DX01 * Y0;
+	int32 C1 = DY12 * X1 - DX12 * Y1;
+	int32 C2 = DY20 * X2 - DX20 * Y2;
+
+	// Correct for fill convention
+	C0 += ( DY01 < 0 || ( DY01 == 0 && DX01 > 0 ) ) ? 0 : -1;
+	C1 += ( DY12 < 0 || ( DY12 == 0 && DX12 > 0 ) ) ? 0 : -1;
+	C2 += ( DY20 < 0 || ( DY20 == 0 && DX20 > 0 ) ) ? 0 : -1;
+
+	// Dilate edges
+	C0 += ( abs(DX01) + abs(DY01) ) * Dilate;
+	C1 += ( abs(DX12) + abs(DY12) ) * Dilate;
+	C2 += ( abs(DX20) + abs(DY20) ) * Dilate;
+
+	for( int32 y = MinY; y < MaxY; y++ )
+	{
+		for( int32 x = MinX; x < MaxX; x++ )
+		{
+			// same as Edge1 >= 0 && Edge2 >= 0 && Edge3 >= 0
+			int32 IsInside;
+			IsInside  = C0 + (DX01 * y - DY01 * x) * 16;
+			IsInside |= C1 + (DX12 * y - DY12 * x) * 16;
+			IsInside |= C2 + (DX20 * y - DY20 * x) * 16;
+
+			if( IsInside >= 0 )
+			{
+				Shader.Process( x, y );
+			}
+		}
+	}
+}
+
+void FLayoutUV::RasterizeChart( const FMeshChart& Chart, uint32 RectW, uint32 RectH )
+{
+	// Bilinear footprint is -1 to 1 pixels. If packed geometrically, only a half pixel dilation
+	// would be needed to guarantee all charts were at least 1 pixel away, safe for bilinear filtering.
+	// Unfortunately, with pixel packing a full 1 pixel dilation is required unless chart edges exactly
+	// align with pixel centers.
+
+	ChartRaster.Clear();
+
+	for( uint32 Tri = Chart.FirstTri; Tri < Chart.LastTri; Tri++ )
+	{
+		FVector2D Points[3];
+		for( int k = 0; k < 3; k++ )
+		{
+			const FVector2D& UV = TexCoords[ 3 * SortedTris[ Tri ] + k ];
+			Points[k] = UV.X * Chart.PackingScaleU + UV.Y * Chart.PackingScaleV + Chart.PackingBias;
+		}
+
+		RasterizeTriangle< FAllocator2DShader, 16 >( ChartShader, Points, RectW, RectH );
+	}
+
+	if ( LayoutVersion >= ELightmapUVVersion::Segments )
+	{
+		ChartRaster.CreateUsedSegments();
+	}
+}
+
+void FLayoutUV::CommitPackedUVs()
+{
+	// Alloc new UV channel
+	MeshView.InitOutputTexcoords(TexCoords.Num());
+
+	// Commit chart UVs
+	for( int32 i = 0; i < Charts.Num(); i++ )
+	{
+		FMeshChart& Chart = Charts[i];
+
+		Chart.PackingScaleU /= TextureResolution;
+		Chart.PackingScaleV /= TextureResolution;
+		Chart.PackingBias /= TextureResolution;
+
+		for( uint32 Tri = Chart.FirstTri; Tri < Chart.LastTri; Tri++ )
+		{
+			for( int k = 0; k < 3; k++ )
+			{
+				uint32 Index = 3 * SortedTris[ Tri ] + k;
+				const FVector2D& UV = TexCoords[ Index ];
+				FVector2D TransformedUV = UV.X * Chart.PackingScaleU + UV.Y * Chart.PackingScaleV + Chart.PackingBias;
+				MeshView.SetOutputTexcoord(Index, TransformedUV);
+			}
+		}
+	}
+}
+
+inline bool FLayoutUV::PositionsMatch( uint32 a, uint32 b ) const
+{
+	return ( MeshView.GetPosition(a) - MeshView.GetPosition(b) ).IsNearlyZero( THRESH_POINTS_ARE_SAME );
+}
+
+inline bool FLayoutUV::NormalsMatch( uint32 a, uint32 b ) const
+{
+	return ( MeshView.GetNormal(a) - MeshView.GetNormal(b) ).IsNearlyZero( THRESH_NORMALS_ARE_SAME );
+}
+
+inline bool FLayoutUV::UVsMatch( uint32 a, uint32 b ) const
+{
+	return ( MeshView.GetInputTexcoord(a) - MeshView.GetInputTexcoord(b) ).IsNearlyZero(UVLAYOUT_THRESH_UVS_ARE_SAME);
+}
+
+inline bool FLayoutUV::VertsMatch( uint32 a, uint32 b ) const
+{
+	return PositionsMatch( a, b ) && UVsMatch( a, b );
+}
+
+// Signed UV area
+inline float FLayoutUV::TriangleUVArea( uint32 Tri ) const
+{
+	FVector2D UVs[3];
+	for( int k = 0; k < 3; k++ )
+	{
+		UVs[k] = MeshView.GetInputTexcoord(3 * Tri + k);
+	}
+
+	FVector2D EdgeUV1 = UVs[1] - UVs[0];
+	FVector2D EdgeUV2 = UVs[2] - UVs[0];
+	return 0.5f * ( EdgeUV1.X * EdgeUV2.Y - EdgeUV1.Y * EdgeUV2.X );
+}
+
+inline void FLayoutUV::DisconnectChart( FMeshChart& Chart, uint32 Side )
+{
+	if( Chart.Join[ Side ] != -1 )
+	{
+		Charts[ Chart.Join[ Side ] ].Join[ Side ^ 1 ] = -1;
+		Chart.Join[ Side ] = -1;
+	}
+}
+
+inline float FLayoutUV::GetUVEqualityThreshold() const
+{
+	return LayoutVersion >= ELightmapUVVersion::SmallChartPacking ? NEW_UVS_ARE_SAME : LEGACY_UVS_ARE_SAME;
+}
diff --git a/Engine/Source/Developer/MeshUtilitiesCommon/Private/OverlappingCorners.cpp b/Engine/Source/Developer/MeshUtilitiesCommon/Private/OverlappingCorners.cpp
new file mode 100644
--- /dev/null
+++ b/Engine/Source/Developer/MeshUtilitiesCommon/Private/OverlappingCorners.cpp
@@ -0,0 +1,148 @@
+// Copyright 1998-2018 Epic Games, Inc. All Rights Reserved.
+
+#include "OverlappingCorners.h"
+#include "MeshUtilitiesCommon.h"
+
+FOverlappingCorners::FOverlappingCorners(const TArray<FVector>& InVertices, const TArray<uint32>& InIndices, float ComparisonThreshold)
+{
+	const int32 NumWedges = InIndices.Num();
+
+	// Create a list of vertex Z/index pairs
+	TArray<FIndexAndZ> VertIndexAndZ;
+	VertIndexAndZ.Reserve(NumWedges);
+	for (int32 WedgeIndex = 0; WedgeIndex < NumWedges; WedgeIndex++)
+	{
+		new(VertIndexAndZ)FIndexAndZ(WedgeIndex, InVertices[InIndices[WedgeIndex]]);
+	}
+
+	// Sort the vertices by z value
+	VertIndexAndZ.Sort(FCompareIndexAndZ());
+
+	Init(NumWedges);
+
+	// Search for duplicates, quickly!
+	for (int32 i = 0; i < VertIndexAndZ.Num(); i++)
+	{
+		// only need to search forward, since we add pairs both ways
+		for (int32 j = i + 1; j < VertIndexAndZ.Num(); j++)
+		{
+			if (FMath::Abs(VertIndexAndZ[j].Z - VertIndexAndZ[i].Z) > ComparisonThreshold)
+				break; // can't be any more dups
+
+			const FVector& PositionA = InVertices[InIndices[VertIndexAndZ[i].Index]];
+			const FVector& PositionB = InVertices[InIndices[VertIndexAndZ[j].Index]];
+
+			if (PointsEqual(PositionA, PositionB, ComparisonThreshold))
+			{
+				Add(VertIndexAndZ[i].Index, VertIndexAndZ[j].Index);
+			}
+		}
+	}
+
+	FinishAdding();
+}
+
+void FOverlappingCorners::Init(int32 NumIndices)
+{
+	Arrays.Reset();
+	Sets.Reset();
+	bFinishedAdding = false;
+
+	IndexBelongsTo.Reset(NumIndices);
+	IndexBelongsTo.AddUninitialized(NumIndices);
+	FMemory::Memset(IndexBelongsTo.GetData(), 0xFF, NumIndices * sizeof(int32));
+}
+
+void FOverlappingCorners::Add(int32 Key, int32 Value)
+{
+	check(Key != Value);
+	check(bFinishedAdding == false);
+
+	int32 ContainerIndex = IndexBelongsTo[Key];
+	if (ContainerIndex == INDEX_NONE)
+	{
+		ContainerIndex = Arrays.Num();
+		TArray<int32>& Container = Arrays.AddDefaulted_GetRef();
+		Container.Reserve(6);
+		Container.Add(Key);
+		Container.Add(Value);
+		IndexBelongsTo[Key] = ContainerIndex;
+		IndexBelongsTo[Value] = ContainerIndex;
+	}
+	else
+	{
+		IndexBelongsTo[Value] = ContainerIndex;
+
+		TArray<int32>& ArrayContainer = Arrays[ContainerIndex];
+		if (ArrayContainer.Num() == 1)
+		{
+			// Container is a set
+			Sets[ArrayContainer.Last()].Add(Value);
+		}
+		else
+		{
+			// Container is an array
+			ArrayContainer.AddUnique(Value);
+
+			// Change container into set when one vertex is shared by large number of triangles
+			if (ArrayContainer.Num() > 12)
+			{
+				int32 SetIndex = Sets.Num();
+				TSet<int32>& Set = Sets.AddDefaulted_GetRef();
+				Set.Append(ArrayContainer);
+
+				// Having one element means we are using a set
+				// An array will never have just 1 element normally because we add them as pairs
+				ArrayContainer.Reset(1);
+				ArrayContainer.Add(SetIndex);
+			}
+		}
+	}
+}
+
+void FOverlappingCorners::FinishAdding()
+{
+	check(bFinishedAdding == false);
+
+	for (TArray<int32>& Array : Arrays)
+	{
+		// Turn sets back into arrays for easier iteration code
+		// Also reduces peak memory later in the import process
+		if (Array.Num() == 1)
+		{
+			TSet<int32>& Set = Sets[Array.Last()];
+			Array.Reset(Set.Num());
+			for (int32 i : Set)
+			{
+				Array.Add(i);
+			}
+		}
+
+		// Sort arrays now to avoid sort multiple times
+		Array.Sort();
+	}
+
+	Sets.Empty();
+
+	bFinishedAdding = true;
+}
+
+uint32 FOverlappingCorners::GetAllocatedSize(void) const
+{
+	uint32 BaseMemoryAllocated = IndexBelongsTo.GetAllocatedSize() + Arrays.GetAllocatedSize() + Sets.GetAllocatedSize();
+
+	uint32 ArraysMemory = 0;
+	for (const TArray<int32>& ArrayIt : Arrays)
+	{
+		ArraysMemory += ArrayIt.GetAllocatedSize();
+	}
+
+	uint32 SetsMemory = 0;
+	for (const TSet<int32>& SetsIt : Sets)
+	{
+		SetsMemory += SetsIt.GetAllocatedSize();
+	}
+
+	return BaseMemoryAllocated + ArraysMemory + SetsMemory;
+}
+
diff --git a/Engine/Source/Developer/MeshUtilitiesCommon/Public/Allocator2D.h b/Engine/Source/Developer/MeshUtilitiesCommon/Public/Allocator2D.h
new file mode 100644
--- /dev/null
+++ b/Engine/Source/Developer/MeshUtilitiesCommon/Public/Allocator2D.h
@@ -0,0 +1,172 @@
+// Copyright 1998-2018 Epic Games, Inc. All Rights Reserved.
+
+#pragma once
+
+#include "CoreMinimal.h"
+
+struct FRect;
+struct Rect;
+
+class MESHUTILITIESCOMMON_API FAllocator2D
+{
+public:
+	struct FRect
+	{
+		uint32 X;
+		uint32 Y;
+		uint32 W;
+		uint32 H;
+	};
+
+	struct FSegment
+	{
+		uint32 StartPos;
+		uint32 Length;
+
+		bool operator<( const FSegment& Other ) const { return StartPos < Other.StartPos; }
+	};
+
+	struct FRow
+	{
+		uint32 Index;
+		uint32 LongestSegment; // Represents either the longest free segment or the longest used segment, depending on how we're using this row
+
+		TArray< FSegment > FreeSegments;
+		TArray< FSegment > UsedSegments;
+	};
+
+public:
+				FAllocator2D( uint32 Width, uint32 Height );
+				FAllocator2D( const FAllocator2D& Other );
+				~FAllocator2D();
+	
+	FAllocator2D& operator=( const FAllocator2D& Other );
+
+	// Must clear before using
+	void		Clear();
+
+	bool		Find( FRect& Rect );
+	bool		Test( FRect Rect );
+	void		Alloc( FRect Rect );
+
+	bool		FindBitByBit( FRect& Rect, const FAllocator2D& Other );
+	bool		FindWithSegments( FRect& Rect, FRect BestRect, const FAllocator2D& Other );
+	bool		Test( FRect Rect, const FAllocator2D& Other );
+	void		Alloc( FRect Rect, const FAllocator2D& Other );
+	
+	uint64		GetBit( uint32 x, uint32 y ) const;
+	void		SetBit( uint32 x, uint32 y );
+	void		ClearBit( uint32 x, uint32 y );
+	
+	void		CreateUsedSegments();
+	void		MergeSegments( FRect Rect, const FAllocator2D& Other );
+
+	void		FlipX( FRect Rect );
+	void		FlipY( FRect Rect );
+
+protected:
+	bool		TestAllRows( FRect Rect, const FAllocator2D& Other, uint32& FailedLength );
+	bool		TestRow( const FRow& ThisRow, const FRow& OtherRow, FRect Rect, uint32& FailedLength );
+
+	void		InitSegments();
+	void		AddUsedSegment( FRow& Row, uint32 StartPos, uint32 Length );
+
+private:
+	uint64*		Bits;
+
+	uint32		Width;
+	uint32		Height;
+	uint32		Pitch;
+
+	TArray< FRow > Rows;
+	int32		LastRowFail;
+};
+
+struct FAllocator2DShader
+{
+	FAllocator2D*	Allocator2D;
+
+	FAllocator2DShader( FAllocator2D* InAllocator2D )
+	: Allocator2D( InAllocator2D )
+	{}
+
+	FORCEINLINE void Process( uint32 x, uint32 y )
+	{
+		Allocator2D->SetBit( x, y );
+	}
+};
+
+// Returns non-zero if set
+FORCEINLINE uint64 FAllocator2D::GetBit( uint32 x, uint32 y ) const
+{
+	return Bits[ (x >> 6) + y * Pitch ] & ( 1ull << ( x & 63 ) );
+}
+
+FORCEINLINE void FAllocator2D::SetBit( uint32 x, uint32 y )
+{
+	Bits[ (x >> 6) + y * Pitch ] |= ( 1ull << ( x & 63 ) );
+}
+
+FORCEINLINE void FAllocator2D::ClearBit( uint32 x, uint32 y )
+{
+	Bits[ (x >> 6) + y * Pitch ] &= ~( 1ull << ( x & 63 ) );
+}
+
+inline bool FAllocator2D::Test( FRect Rect )
+{
+	for( uint32 y = Rect.Y; y < Rect.Y + Rect.H; y++ )
+	{
+		for( uint32 x = Rect.X; x < Rect.X + Rect.W; x++ )
+		{
+			if( GetBit( x, y ) )
+			{
+				return false;
+			}
+		}
+	}
+	
+	return true;
+}
+
+inline bool FAllocator2D::Test( FRect Rect, const FAllocator2D& Other )
+{
+	const uint32 LowShift = Rect.X & 63;
+	const uint32 HighShift = 64 - LowShift;
+
+	for( uint32 y = 0; y < Rect.H; y++ )
+	{
+#if 1
+		uint32 ThisIndex = (Rect.X >> 6) + (y + Rect.Y) * Pitch;
+		uint32 OtherIndex = y * Pitch;
+
+		// Test a uint64 at a time
+		for( uint32 x = 0; x < Rect.W; x += 64 )
+		{
+			// no need to zero out HighInt on wrap around because Other will always be zero outside Rect.
+			uint64 LowInt  = Bits[ ThisIndex ];
+			uint64 HighInt = Bits[ ThisIndex + 1 ];
+
+			uint64 ThisInt = (HighInt << HighShift) | (LowInt >> LowShift);
+			uint64 OtherInt = Other.Bits[ OtherIndex ];
+
+			if( ThisInt & OtherInt )
+			{
+				return false;
+			}
+
+			ThisIndex++;
+			OtherIndex++;
+		}
+#else
+		for( uint32 x = 0; x < Rect.W; x++ )
+		{
+			if( Other.GetBit( x, y ) && GetBit( x + Rect.X, y + Rect.Y ) )
+			{
+				return false;
+			}
+		}
+#endif
+	}
+	
+	return true;
+}
diff --git a/Engine/Source/Developer/MeshUtilitiesCommon/Public/DisjointSet.h b/Engine/Source/Developer/MeshUtilitiesCommon/Public/DisjointSet.h
new file mode 100644
--- /dev/null
+++ b/Engine/Source/Developer/MeshUtilitiesCommon/Public/DisjointSet.h
@@ -0,0 +1,82 @@
+// Copyright 1998-2018 Epic Games, Inc. All Rights Reserved.
+
+#pragma once
+
+#include "CoreMinimal.h"
+
+class FDisjointSet
+{
+public:
+			FDisjointSet( const uint32 Size );
+	
+	void	Union( uint32 x, uint32 y );
+	uint32	Find( uint32 i );
+
+	uint32	operator[]( uint32 i ) const	{ return Parents[i]; }
+
+private:
+	TArray<	uint32 >	Parents;
+};
+
+FDisjointSet::FDisjointSet( const uint32 Size )
+{
+	Parents.SetNumUninitialized( Size );
+	for( uint32 i = 0; i < Size; i++ )
+	{
+		Parents[i] = i;
+	}
+}
+
+// Union with splicing
+void FDisjointSet::Union( uint32 x, uint32 y )
+{
+	while( Parents[x] != Parents[y] )
+	{
+		// Pick larger
+		if( Parents[x] < Parents[y] )
+		{
+			uint32 p = Parents[x];
+			Parents[x] = Parents[y];
+			if( x == p )
+			{
+				return;
+			}
+			x = p;
+		}
+		else
+		{
+			uint32 p = Parents[y];
+			Parents[y] = Parents[x];
+			if( y == p )
+			{
+				return;
+			}
+			y = p;
+		}
+	}
+}
+
+// Find with path compression
+uint32 FDisjointSet::Find( uint32 i )
+{
+	// Find root
+	uint32 Start = i;
+	uint32 Root = Parents[i];
+	while( Root != i )
+	{
+		i = Root;
+		Root = Parents[i];
+	}
+
+	// Point all nodes on path to root
+	i = Start;
+	uint32 Parent = Parents[i];
+	while( Parent != Root )
+	{
+		Parents[i] = Root;
+		i = Parent;
+		Parent = Parents[i];
+	}
+
+	return Root;
+}
diff --git a/Engine/Source/Developer/MeshUtilitiesCommon/Public/LayoutUV.h b/Engine/Source/Developer/MeshUtilitiesCommon/Public/LayoutUV.h
new file mode 100644
--- /dev/null
+++ b/Engine/Source/Developer/MeshUtilitiesCommon/Public/LayoutUV.h
@@ -0,0 +1,96 @@
+// Copyright 1998-2018 Epic Games, Inc. All Rights Reserved.
+
+#pragma once
+
+#include "CoreMinimal.h"
+#include "MeshUtilitiesCommon.h"
+#include "Allocator2D.h"
+
+struct FMeshChart
+{
+	uint32		FirstTri;
+	uint32		LastTri;
+	
+	FVector2D	MinUV;
+	FVector2D	MaxUV;
+	
+	float		UVArea;
+	FVector2D	UVScale;
+	FVector2D	WorldScale;
+	
+	FVector2D	PackingScaleU;
+	FVector2D	PackingScaleV;
+	FVector2D	PackingBias;
+
+	int32		Join[4];
+
+	int32		Id; // Store a unique id so that we can come back to the initial Charts ordering when necessary
+};
+
+struct FOverlappingCorners;
+
+class MESHUTILITIESCOMMON_API FLayoutUV
+{
+public:
+
+	/**
+	 * Abstract triangle mesh view interface that may be used by any module without introducing
+	 * a dependency on a concrete mesh type (and thus potentially circular module references).
+	 * This abstraction results in a performance penalty due to virtual dispatch,
+	 * however it is expected to be insignificant compared to the rest of work done by FLayoutUV
+	 * and cache misses due to indexed vertex data access.
+	*/
+	struct IMeshView
+	{
+		virtual ~IMeshView() {}
+
+		virtual uint32      GetNumIndices() const = 0;
+		virtual FVector     GetPosition(uint32 Index) const = 0;
+		virtual FVector     GetNormal(uint32 Index) const = 0;
+		virtual FVector2D   GetInputTexcoord(uint32 Index) const = 0;
+
+		virtual void        InitOutputTexcoords(uint32 Num) = 0;
+		virtual void        SetOutputTexcoord(uint32 Index, const FVector2D& Value) = 0;
+	};
+
+				FLayoutUV( IMeshView& InMeshView, uint32 InTextureResolution );
+
+	int32		FindCharts( const FOverlappingCorners& OverlappingCorners );
+	bool		FindBestPacking();
+	void		CommitPackedUVs();
+
+	void		SetVersion( ELightmapUVVersion Version ) { LayoutVersion = Version; }
+
+private:
+	bool		PositionsMatch( uint32 a, uint32 b ) const;
+	bool		NormalsMatch( uint32 a, uint32 b ) const;
+	bool		UVsMatch( uint32 a, uint32 b ) const;
+	bool		VertsMatch( uint32 a, uint32 b ) const;
+	float		TriangleUVArea( uint32 Tri ) const;
+	void		DisconnectChart( FMeshChart& Chart, uint32 Side );
+
+	void		ScaleCharts( float UVScale );
+	bool		PackCharts();
+	void		OrientChart( FMeshChart& Chart, int32 Orientation );
+	void		RasterizeChart( const FMeshChart& Chart, uint32 RectW, uint32 RectH );
+
+	float		GetUVEqualityThreshold() const;
+
+	IMeshView&	MeshView;
+	uint32		TextureResolution;
+	
+	TArray< FVector2D >		TexCoords;
+	TArray< uint32 >		SortedTris;
+	TArray< FMeshChart >	Charts;
+	float					TotalUVArea;
+	float					MaxChartSize;
+
+	FAllocator2D		LayoutRaster;
+	FAllocator2D		ChartRaster;
+	FAllocator2D		BestChartRaster;
+	FAllocator2DShader	ChartShader;
+
+	ELightmapUVVersion	LayoutVersion;
+
+	int32				NextMeshChartId;
+};
diff --git a/Engine/Source/Developer/MeshUtilitiesCommon/Public/MeshUtilitiesCommon.h b/Engine/Source/Developer/MeshUtilitiesCommon/Public/MeshUtilitiesCommon.h
new file mode 100644
--- /dev/null
+++ b/Engine/Source/Developer/MeshUtilitiesCommon/Public/MeshUtilitiesCommon.h
@@ -0,0 +1,52 @@
+// Copyright 1998-2018 Epic Games, Inc. All Rights Reserved.
+
+#pragma once
+
+#include "CoreMinimal.h"
+
+enum class ELightmapUVVersion : int32
+{
+	BitByBit = 0,
+	Segments = 1,
+	SmallChartPacking = 2,
+	ScaleChartsOrderingFix = 3,
+	ChartJoiningLFix = 4,
+	Latest = ChartJoiningLFix
+};
+
+/** Helper struct for building acceleration structures. */
+struct FIndexAndZ
+{
+	float Z;
+	int32 Index;
+
+	/** Default constructor. */
+	FIndexAndZ() {}
+
+	/** Initialization constructor. */
+	FIndexAndZ(int32 InIndex, FVector V)
+	{
+		Z = 0.30f * V.X + 0.33f * V.Y + 0.37f * V.Z;
+		Index = InIndex;
+	}
+};
+
+/** Sorting function for vertex Z/index pairs. */
+struct FCompareIndexAndZ
+{
+	FORCEINLINE bool operator()(FIndexAndZ const& A, FIndexAndZ const& B) const { return A.Z < B.Z; }
+};
+
+/**
+* Returns true if the specified points are about equal
+*/
+inline bool PointsEqual(const FVector& V1, const FVector& V2, float ComparisonThreshold)
+{
+	if (FMath::Abs(V1.X - V2.X) > ComparisonThreshold
+		|| FMath::Abs(V1.Y - V2.Y) > ComparisonThreshold
+		|| FMath::Abs(V1.Z - V2.Z) > ComparisonThreshold)
+	{
+		return false;
+	}
+	return true;
+}
diff --git a/Engine/Source/Developer/MeshUtilitiesCommon/Public/OverlappingCorners.h b/Engine/Source/Developer/MeshUtilitiesCommon/Public/OverlappingCorners.h
new file mode 100644
--- /dev/null
+++ b/Engine/Source/Developer/MeshUtilitiesCommon/Public/OverlappingCorners.h
@@ -0,0 +1,44 @@
+// Copyright 1998-2018 Epic Games, Inc. All Rights Reserved.
+
+#pragma once
+
+#include "CoreMinimal.h"
+
+/**
+* Container to hold overlapping corners. For a vertex, lists all the overlapping vertices
+*/
+struct MESHUTILITIESCOMMON_API FOverlappingCorners
+{
+	FOverlappingCorners() {}
+
+	FOverlappingCorners(const TArray<FVector>& InVertices, const TArray<uint32>& InIndices, float ComparisonThreshold);
+
+	/* Resets, pre-allocates memory, marks all indices as not overlapping in preperation for calls to Add() */
+	void Init(int32 NumIndices);
+
+	/* Add overlapping indices pair */
+	void Add(int32 Key, int32 Value);
+
+	/* Sorts arrays, converts sets to arrays for sorting and to allow simple iterating code, prevents additional adding */
+	void FinishAdding();
+
+	/* Estimate memory allocated */
+	uint32 GetAllocatedSize(void) const;
+
+	/**
+	* @return array of sorted overlapping indices including input 'Key', empty array for indices that have no overlaps.
+	*/
+	const TArray<int32>& FindIfOverlapping(int32 Key) const
+	{
+		check(bFinishedAdding);
+		int32 ContainerIndex = IndexBelongsTo[Key];
+		return (ContainerIndex != INDEX_NONE) ? Arrays[ContainerIndex] : EmptyArray;
+	}
+
+private:
+	TArray<int32> IndexBelongsTo;
+	TArray< TArray<int32> > Arrays;
+	TArray< TSet<int32> > Sets;
+	TArray<int32> EmptyArray;
+	bool bFinishedAdding = false;
+};
diff --git a/Engine/Source/Developer/SimplygonMeshReduction/Private/SimplygonMeshReduction.cpp b/Engine/Source/Developer/SimplygonMeshReduction/Private/SimplygonMeshReduction.cpp
--- a/Engine/Source/Developer/SimplygonMeshReduction/Private/SimplygonMeshReduction.cpp
+++ b/Engine/Source/Developer/SimplygonMeshReduction/Private/SimplygonMeshReduction.cpp
@@ -566,7 +566,7 @@ public:
 		FMeshDescription& OutReducedMesh,
 		float& OutMaxDeviation,
 		const FMeshDescription& InMesh,
-		const TMultiMap<int32, int32>& InOverlappingCorners,
+		const FOverlappingCorners& InOverlappingCorners,
 		const struct FMeshReductionSettings& ReductionSettings
 	) override
 	{
diff --git a/Engine/Source/Editor/UnrealEd/Private/Fbx/FbxSkeletalMeshImport.cpp b/Engine/Source/Editor/UnrealEd/Private/Fbx/FbxSkeletalMeshImport.cpp
--- a/Engine/Source/Editor/UnrealEd/Private/Fbx/FbxSkeletalMeshImport.cpp
+++ b/Engine/Source/Editor/UnrealEd/Private/Fbx/FbxSkeletalMeshImport.cpp
@@ -41,6 +41,7 @@
 
 #include "ApexClothingUtils.h"
 #include "MeshUtilities.h"
+#include "OverlappingCorners.h"
 
 #include "IMessageLogListing.h"
 #include "MessageLogModule.h"
diff --git a/Engine/Source/Editor/UnrealEd/UnrealEd.Build.cs b/Engine/Source/Editor/UnrealEd/UnrealEd.Build.cs
--- a/Engine/Source/Editor/UnrealEd/UnrealEd.Build.cs
+++ b/Engine/Source/Editor/UnrealEd/UnrealEd.Build.cs
@@ -123,6 +123,7 @@ public class UnrealEd : ModuleRules
 				"PropertyEditor",
 				"Projects",
 				"RawMesh",
+				"MeshUtilitiesCommon",
 				"RenderCore",
 				"RHI",
 				"ShaderCore",
diff --git a/Engine/Source/Runtime/Engine/Engine.Build.cs b/Engine/Source/Runtime/Engine/Engine.Build.cs
--- a/Engine/Source/Runtime/Engine/Engine.Build.cs
+++ b/Engine/Source/Runtime/Engine/Engine.Build.cs
@@ -219,6 +219,8 @@ public class Engine : ModuleRules
 			if (Target.Type != TargetType.Server)
 			{
 				PrivateIncludePathModuleNames.Add("MeshUtilities");
+				PrivateIncludePathModuleNames.Add("MeshUtilitiesCommon");
+
 				DynamicallyLoadedModuleNames.Add("MeshUtilities");
 
 				PrivateDependencyModuleNames.AddRange(
diff --git a/Engine/Source/Runtime/Engine/Private/StaticMesh.cpp b/Engine/Source/Runtime/Engine/Private/StaticMesh.cpp
--- a/Engine/Source/Runtime/Engine/Private/StaticMesh.cpp
+++ b/Engine/Source/Runtime/Engine/Private/StaticMesh.cpp
@@ -46,6 +46,7 @@
 #include "Settings/EditorExperimentalSettings.h"
 #include "MeshBuilder.h"
 #include "MeshUtilities.h"
+#include "MeshUtilitiesCommon.h"
 #include "DerivedDataCacheInterface.h"
 #include "PlatformInfo.h"
 #include "ScopedTransaction.h"
diff --git a/Engine/Source/Runtime/Engine/Private/StaticMeshBuild.cpp b/Engine/Source/Runtime/Engine/Private/StaticMeshBuild.cpp
--- a/Engine/Source/Runtime/Engine/Private/StaticMeshBuild.cpp
+++ b/Engine/Source/Runtime/Engine/Private/StaticMeshBuild.cpp
@@ -16,6 +16,7 @@
 
 #if WITH_EDITOR
 #include "MeshUtilities.h"
+#include "MeshUtilitiesCommon.h"
 #include "Misc/FeedbackContext.h"
 #include "Misc/App.h"
 #endif // #if WITH_EDITOR
