Copy of 4923653, 4923695 from //UE4/Dev-Enterprise:

UV packing issue for geometries with high chart count

diff --git a/Engine/Source/Developer/MeshDescriptionOperations/Private/MeshDescriptionOperations.cpp b/Engine/Source/Developer/MeshDescriptionOperations/Private/MeshDescriptionOperations.cpp
--- a/Engine/Source/Developer/MeshDescriptionOperations/Private/MeshDescriptionOperations.cpp
+++ b/Engine/Source/Developer/MeshDescriptionOperations/Private/MeshDescriptionOperations.cpp
@@ -1412,7 +1412,16 @@ struct FLayoutUVMeshDescriptionView final : FLayoutUV::IMeshView
 	}
 };
 
-void FMeshDescriptionOperations::CreateLightMapUVLayout(FMeshDescription& MeshDescription,
+int32 FMeshDescriptionOperations::GetUVChartCount(FMeshDescription& MeshDescription, int32 SrcLightmapIndex, ELightmapUVVersion LightmapUVVersion, const FOverlappingCorners& OverlappingCorners)
+{
+	uint32 UnusedDstIndex = -1;
+	FLayoutUVMeshDescriptionView MeshDescriptionView(MeshDescription, SrcLightmapIndex, UnusedDstIndex);
+	FLayoutUV Packer(MeshDescriptionView);
+	Packer.SetVersion(LightmapUVVersion);
+	return Packer.FindCharts(OverlappingCorners);
+}
+
+bool FMeshDescriptionOperations::CreateLightMapUVLayout(FMeshDescription& MeshDescription,
 	int32 SrcLightmapIndex,
 	int32 DstLightmapIndex,
 	int32 MinLightmapResolution,
@@ -1420,15 +1429,16 @@ void FMeshDescriptionOperations::CreateLightMapUVLayout(FMeshDescription& MeshDe
 	const FOverlappingCorners& OverlappingCorners)
 {
 	FLayoutUVMeshDescriptionView MeshDescriptionView(MeshDescription, SrcLightmapIndex, DstLightmapIndex);
-	FLayoutUV Packer(MeshDescriptionView, MinLightmapResolution);
+	FLayoutUV Packer(MeshDescriptionView);
 	Packer.SetVersion(LightmapUVVersion);
 
 	Packer.FindCharts(OverlappingCorners);
-	bool bPackSuccess = Packer.FindBestPacking();
+	bool bPackSuccess = Packer.FindBestPacking(MinLightmapResolution);
 	if (bPackSuccess)
 	{
 		Packer.CommitPackedUVs();
 	}
+	return bPackSuccess;
 }
 
 bool FMeshDescriptionOperations::GenerateUniqueUVsForStaticMesh(const FMeshDescription& MeshDescription, int32 TextureResolution, bool bMergeIdenticalMaterials, TArray<FVector2D>& OutTexCoords)
@@ -1544,10 +1554,10 @@ bool FMeshDescriptionOperations::GenerateUniqueUVsForStaticMesh(const FMeshDescr
 
 	// Generate new UVs
 	FLayoutUVMeshDescriptionView DuplicateMeshDescriptionView(DuplicateMeshDescription, 0, 1);
-	FLayoutUV Packer(DuplicateMeshDescriptionView, FMath::Clamp(TextureResolution / 4, 32, 512));
+	FLayoutUV Packer(DuplicateMeshDescriptionView);
 	Packer.FindCharts(OverlappingCorners);
 
-	bool bPackSuccess = Packer.FindBestPacking();
+	bool bPackSuccess = Packer.FindBestPacking(FMath::Clamp(TextureResolution / 4, 32, 512));
 	if (bPackSuccess)
 	{
 		Packer.CommitPackedUVs();
diff --git a/Engine/Source/Developer/MeshDescriptionOperations/Public/MeshDescriptionOperations.h b/Engine/Source/Developer/MeshDescriptionOperations/Public/MeshDescriptionOperations.h
--- a/Engine/Source/Developer/MeshDescriptionOperations/Public/MeshDescriptionOperations.h
+++ b/Engine/Source/Developer/MeshDescriptionOperations/Public/MeshDescriptionOperations.h
@@ -72,8 +72,25 @@ public:
 
 	/** Find all overlapping vertex using the threshold in the mesh description. */
 	static void FindOverlappingCorners(FOverlappingCorners& OverlappingCorners, const FMeshDescription& MeshDescription, float ComparisonThreshold);
+	
+	/** Find all charts in the mesh description. */
+	static int32 GetUVChartCount(FMeshDescription& MeshDescription, int32 SrcLightmapIndex, ELightmapUVVersion LightmapUVVersion, const FOverlappingCorners& OverlappingCorners);
 
-	static void CreateLightMapUVLayout(FMeshDescription& MeshDescription,
+	/**
+	 * Find and pack UV charts for lightmap.
+	 * The packing algorithm uses a rasterization method, hence the resolution parameter.
+	 *
+	 * If the given minimum resolution is not enough to handle all the charts, generation will fail.
+	 *
+	 * @param MeshDescription        Edited mesh
+	 * @param SrcLightmapIndex       index of the source UV channel
+	 * @param DstLightmapIndex       index of the destination UV channel
+	 * @param MinLightmapResolution  Minimum resolution used for the packing
+	 * @param LightmapUVVersion      Algorithm version
+	 * @param OverlappingCorners     Overlapping corners of the given mesh
+	 * @return                       UV layout correctly generated
+	 */
+	static bool CreateLightMapUVLayout(FMeshDescription& MeshDescription,
 		int32 SrcLightmapIndex,
 		int32 DstLightmapIndex,
 		int32 MinLightmapResolution,
diff --git a/Engine/Source/Developer/MeshUtilities/Private/MeshUtilities.cpp b/Engine/Source/Developer/MeshUtilities/Private/MeshUtilities.cpp
--- a/Engine/Source/Developer/MeshUtilities/Private/MeshUtilities.cpp
+++ b/Engine/Source/Developer/MeshUtilities/Private/MeshUtilities.cpp
@@ -2503,11 +2503,11 @@ public:
 					}
 
 					FLayoutUVRawMeshView RawMeshView(RawMesh, SrcModel.BuildSettings.SrcLightmapIndex, SrcModel.BuildSettings.DstLightmapIndex);
-					FLayoutUV Packer(RawMeshView, SrcModel.BuildSettings.MinLightmapResolution);
+					FLayoutUV Packer(RawMeshView);
 					Packer.SetVersion(LightmapUVVersion);
 
 					Packer.FindCharts(OverlappingCorners);
-					bool bPackSuccess = Packer.FindBestPacking();
+					bool bPackSuccess = Packer.FindBestPacking(SrcModel.BuildSettings.MinLightmapResolution);
 					if (bPackSuccess)
 					{
 						Packer.CommitPackedUVs();
@@ -5665,10 +5665,10 @@ bool FMeshUtilities::GenerateUniqueUVsForSkeletalMesh(const FSkeletalMeshLODMode
 
 	// Generate new UVs
 	FLayoutUVRawMeshView TempMeshView(TempMesh, 0, 1);
-	FLayoutUV Packer(TempMeshView, FMath::Clamp(TextureResolution / 4, 32, 512));
+	FLayoutUV Packer(TempMeshView);
 	Packer.FindCharts(OverlappingCorners);
 
-	bool bPackSuccess = Packer.FindBestPacking();
+	bool bPackSuccess = Packer.FindBestPacking(FMath::Clamp(TextureResolution / 4, 32, 512));
 	if (bPackSuccess)
 	{
 		Packer.CommitPackedUVs();
@@ -6166,10 +6166,10 @@ bool FMeshUtilities::GenerateUniqueUVsForStaticMesh(const FRawMesh& RawMesh, int
 
 	// Generate new UVs
 	FLayoutUVRawMeshView TempMeshView(TempMesh, 0, 1);
-	FLayoutUV Packer(TempMeshView, FMath::Clamp(TextureResolution / 4, 32, 512));
+	FLayoutUV Packer(TempMeshView);
 	Packer.FindCharts(OverlappingCorners);
 
-	bool bPackSuccess = Packer.FindBestPacking();
+	bool bPackSuccess = Packer.FindBestPacking(FMath::Clamp(TextureResolution / 4, 32, 512));
 	if (bPackSuccess)
 	{
 		Packer.CommitPackedUVs();
diff --git a/Engine/Source/Developer/MeshUtilitiesCommon/Private/Allocator2D.cpp b/Engine/Source/Developer/MeshUtilitiesCommon/Private/Allocator2D.cpp
--- a/Engine/Source/Developer/MeshUtilitiesCommon/Private/Allocator2D.cpp
+++ b/Engine/Source/Developer/MeshUtilitiesCommon/Private/Allocator2D.cpp
@@ -8,8 +8,6 @@ FAllocator2D::FAllocator2D( uint32 InWidth, uint32 InHeight )
 	, LastRowFail( -1 )
 {
 	Pitch = ( Width + 63 ) / 64;
-	// alloc +1 to avoid buffer overrun
-	Bits = (uint64*)FMemory::Malloc( (Pitch * Height + 1) * sizeof(uint64) );
 
 	Rows.Reserve( Height );
 
@@ -24,46 +22,11 @@ FAllocator2D::FAllocator2D( uint32 InWidth, uint32 InHeight )
 	Clear();
 }
 
-FAllocator2D::~FAllocator2D()
-{
-	FMemory::Free( Bits );
-}
-
-FAllocator2D::FAllocator2D( const FAllocator2D& Other )
-	: Width( Other.Width )
-	, Height( Other.Height )
-	, Pitch( Other.Pitch )
-	, Rows( Other.Rows )
-	, LastRowFail( -1 )
-{
-	Bits = (uint64*)FMemory::Malloc( (Pitch * Height + 1) * sizeof(uint64) );
-	FMemory::Memcpy( Bits, Other.Bits, (Pitch * Height + 1) * sizeof(uint64) );
-}
-
-FAllocator2D& FAllocator2D::operator=( const FAllocator2D& Other )
-{
-	if ( Width != Other.Width || Height != Other.Height || Pitch != Other.Pitch )
-	{
-		Width = Other.Width;
-		Height = Other.Height;
-		Pitch = Other.Pitch;
-
-		FMemory::Free( Bits );
-		Bits = (uint64*)FMemory::Malloc( (Pitch * Height + 1) * sizeof(uint64) );
-	}
-
-	FMemory::Memcpy( Bits, Other.Bits, (Pitch * Height + 1) * sizeof(uint64) );
-
-	Rows = Other.Rows;
-
-	return *this;
-}
-
 void FAllocator2D::Clear()
 {
 	InitSegments();
-
-	FMemory::Memzero( Bits, (Pitch * Height + 1) * sizeof(uint64) );
+	Bits.Reset();
+	Bits.SetNumZeroed(Pitch * Height + 1); // alloc +1 to avoid buffer overrun
 }
 
 bool FAllocator2D::Find( FRect& Rect )
@@ -98,7 +61,7 @@ bool FAllocator2D::FindBitByBit( FRect& Rect, const FAllocator2D& Other )
 			}
 		}
 	}
-	
+
 	return false;
 }
 
@@ -126,7 +89,7 @@ bool FAllocator2D::FindWithSegments( FRect& Rect, FRect BestRect, const FAllocat
 			}
 		}
 	}
-	
+
 	return false;
 }
 
@@ -292,7 +255,7 @@ void FAllocator2D::FlipX( FRect Rect, ELightmapUVVersion LayoutVersion )
 		for ( uint32 LowX = 0; LowX < ( MaxX + 1 ) / 2; ++LowX )
 		{
 			uint32 HighX = MaxX - LowX;
-		
+
 			const uint64 BitLow = GetBit( LowX, Y );
 			const uint64 BitHigh = GetBit( HighX, Y );
 
@@ -408,7 +371,7 @@ void FAllocator2D::CreateUsedSegments()
 				{
 					FirstUsedX = x;
 				}
-					
+
 				if ( k == Pitch - 1 )
 				{
 					AddUsedSegment( CurrentRow, FirstUsedX, x + 64 - FirstUsedX );
@@ -437,7 +400,7 @@ void FAllocator2D::CreateUsedSegments()
 						{
 							FirstUsedX = SubX;
 						}
-							
+
 						if ( SubX == Width - 1 )
 						{
 							AddUsedSegment( CurrentRow, FirstUsedX, SubX + 1 - FirstUsedX );
@@ -509,7 +472,7 @@ void FAllocator2D::MergeSegments( FRect Rect, const FAllocator2D& Other )
 					{
 						ThisRow.FreeSegments.Add( FirstSegment );
 					}
-					
+
 					if ( SecondSegment.Length > 0 )
 					{
 						ThisRow.FreeSegments.Add( SecondSegment );
diff --git a/Engine/Source/Developer/MeshUtilitiesCommon/Private/LayoutUV.cpp b/Engine/Source/Developer/MeshUtilitiesCommon/Private/LayoutUV.cpp
--- a/Engine/Source/Developer/MeshUtilitiesCommon/Private/LayoutUV.cpp
+++ b/Engine/Source/Developer/MeshUtilitiesCommon/Private/LayoutUV.cpp
@@ -17,19 +17,85 @@ DEFINE_LOG_CATEGORY_STATIC(LogLayoutUV, Warning, All);
 #define LEGACY_UVS_ARE_SAME (1.0f / 1024.0f)
 #define UVLAYOUT_THRESH_UVS_ARE_SAME (GetUVEqualityThreshold())
 
-FLayoutUV::FLayoutUV( IMeshView& InMeshView, uint32 InTextureResolution )
+FLayoutUV::FLayoutUV( IMeshView& InMeshView )
 	: MeshView( InMeshView )
-	, TextureResolution( InTextureResolution )
-	, TotalUVArea( 0.0f )
+	, LayoutVersion( ELightmapUVVersion::Latest )
+	, PackedTextureResolution(0)
+{}
+
+
+/** FIRST PASS: Given a Mesh, build the associated set of charts */
+struct FLayoutUV::FChartFinder
+{
+	FChartFinder(IMeshView& InMeshView, ELightmapUVVersion InLayoutVersion);
+
+	int32 FindCharts( const FOverlappingCorners& OverlappingCorners, TArray< FVector2D >& TexCoords, TArray< uint32 >& SortedTris, TArray< FMeshChart >& Charts );
+
+private:
+	bool PositionsMatch( uint32 a, uint32 b ) const;
+	bool NormalsMatch( uint32 a, uint32 b ) const;
+	bool UVsMatch( uint32 a, uint32 b ) const;
+	bool VertsMatch( uint32 a, uint32 b ) const;
+	float TriangleUVArea( uint32 Tri ) const;
+	void DisconnectChart( TArray< FMeshChart >& Charts, FMeshChart& Chart, uint32 Side );
+	float GetUVEqualityThreshold() const;
+
+private:
+	IMeshView& MeshView;
+	ELightmapUVVersion LayoutVersion;
+
+	int32 NextMeshChartId;
+};
+
+
+/** SECOND PASS: Given a set of charts, pack them in the UV space */
+struct FLayoutUV::FChartPacker
+{
+	FChartPacker(IMeshView& InMeshView, ELightmapUVVersion InLayoutVersion, uint32 TextureResolution);
+
+	bool FindBestPacking(const TArray< FVector2D >& TexCoords, const TArray< uint32 >& SortedTris, TArray< FMeshChart >& AllCharts);
+
+private:
+	void ScaleCharts( TArray< FMeshChart >& Charts, float UVScale );
+	bool PackCharts(TArray< FMeshChart >& Charts, const TArray< FVector2D >& TexCoords, const TArray< uint32 >& SortedTris);
+	void OrientChart( FMeshChart& Chart, int32 Orientation );
+	void RasterizeChart( const FMeshChart& Chart, const TArray< FVector2D >& TexCoords, const TArray< uint32 >& SortedTris, uint32 RectW, uint32 RectH );
+
+private:
+	IMeshView& MeshView;
+	ELightmapUVVersion LayoutVersion;
+
+	uint32 TextureResolution;
+	FAllocator2D LayoutRaster;
+	FAllocator2D ChartRaster;
+	FAllocator2D BestChartRaster;
+	float TotalUVArea;
+};
+
+FLayoutUV::FChartPacker::FChartPacker(IMeshView& InMeshView, ELightmapUVVersion InLayoutVersion, uint32 TextureResolution)
+	: MeshView(InMeshView)
+	, LayoutVersion(InLayoutVersion)
+	, TextureResolution(TextureResolution)
 	, LayoutRaster( TextureResolution, TextureResolution )
 	, ChartRaster( TextureResolution, TextureResolution )
 	, BestChartRaster( TextureResolution, TextureResolution )
-	, ChartShader( &ChartRaster )
-	, LayoutVersion( ELightmapUVVersion::Latest )
+	, TotalUVArea(0.0f)
+{
+}
+
+int32 FLayoutUV::FindCharts(const FOverlappingCorners& OverlappingCorners)
+{
+	FChartFinder Finder(MeshView, LayoutVersion);
+	return Finder.FindCharts(OverlappingCorners, TexCoords, SortedTris, Charts);
+}
+
+FLayoutUV::FChartFinder::FChartFinder(IMeshView& InMeshView, ELightmapUVVersion InLayoutVersion)
+	: MeshView(InMeshView)
+	, LayoutVersion(InLayoutVersion)
 	, NextMeshChartId( 0 )
 {}
 
-int32 FLayoutUV::FindCharts( const FOverlappingCorners& OverlappingCorners )
+int32 FLayoutUV::FChartFinder::FindCharts( const FOverlappingCorners& OverlappingCorners, TArray< FVector2D >& TexCoords, TArray< uint32 >& SortedTris, TArray< FMeshChart >& Charts )
 {
 	double Begin = FPlatformTime::Seconds();
 
@@ -231,8 +297,6 @@ int32 FLayoutUV::FindCharts( const FOverlappingCorners& OverlappingCorners )
 				Chart.WorldScale = FVector2D::ZeroVector;
 			}
 		}		
-
-		TotalUVArea += Chart.UVArea * Chart.WorldScale.X * Chart.WorldScale.Y;
 #endif
 	}
 
@@ -504,8 +568,8 @@ int32 FLayoutUV::FindCharts( const FOverlappingCorners& OverlappingCorners )
 					ChartB.LastTri = 0;
 					ChartB.UVArea = 0.0f;
 
-					DisconnectChart( ChartB, Side ^ 2 );
-					DisconnectChart( ChartB, Side ^ 3 );
+					DisconnectChart( Charts, ChartB, Side ^ 2 );
+					DisconnectChart( Charts, ChartB, Side ^ 3 );
 				}
 
 				ChartA.FirstTri = FirstTri;
@@ -555,8 +619,6 @@ int32 FLayoutUV::FindCharts( const FOverlappingCorners& OverlappingCorners )
 				Chart.WorldScale = FVector2D::ZeroVector;
 			}
 		}
-
-		TotalUVArea += Chart.UVArea * Chart.WorldScale.X * Chart.WorldScale.Y;
 	}
 #endif
 
@@ -567,14 +629,24 @@ int32 FLayoutUV::FindCharts( const FOverlappingCorners& OverlappingCorners )
 	return Charts.Num();
 }
 
-bool FLayoutUV::FindBestPacking()
+bool FLayoutUV::FChartPacker::FindBestPacking(const TArray< FVector2D >& TexCoords, const TArray< uint32 >& SortedTris, TArray< FMeshChart >& Charts)
 {
-	if( (uint32)Charts.Num() > TextureResolution * TextureResolution || TotalUVArea == 0.f )
+	if( (uint32)Charts.Num() > TextureResolution * TextureResolution )
 	{
 		// More charts than texels
 		return false;
 	}
 	
+	TotalUVArea = 0.0f;
+	for (const FMeshChart& Chart : Charts)
+	{
+		TotalUVArea += Chart.UVArea * Chart.WorldScale.X * Chart.WorldScale.Y;
+	}
+	if( TotalUVArea <= 0.0f )
+	{
+		return false;
+	}
+
 	const float LinearSearchStart = 0.5f;
 	const float LinearSearchStep = 0.5f;
 	const int32 BinarySearchSteps = 6;
@@ -585,9 +657,9 @@ bool FLayoutUV::FindBestPacking()
 	// Linear search for first fit
 	while(1)
 	{
-		ScaleCharts( UVScalePass );
+		ScaleCharts( Charts, UVScalePass );
 
-		bool bFit = PackCharts();
+		bool bFit = PackCharts(Charts, TexCoords, SortedTris);
 		if( bFit )
 		{
 			break;
@@ -601,9 +673,9 @@ bool FLayoutUV::FindBestPacking()
 	for( int32 i = 0; i < BinarySearchSteps; i++ )
 	{
 		float UVScale = 0.5f * ( UVScaleFail + UVScalePass );
-		ScaleCharts( UVScale );
+		ScaleCharts( Charts, UVScale );
 
-		bool bFit = PackCharts();
+		bool bFit = PackCharts(Charts, TexCoords, SortedTris);
 		if( bFit )
 		{
 			UVScalePass = UVScale;
@@ -615,13 +687,13 @@ bool FLayoutUV::FindBestPacking()
 	}
 
 	// TODO store packing scale/bias separate so this isn't necessary
-	ScaleCharts( UVScalePass );
-	PackCharts();
+	ScaleCharts( Charts, UVScalePass );
+	PackCharts(Charts, TexCoords, SortedTris);
 
 	return true;
 }
 
-void FLayoutUV::ScaleCharts( float UVScale )
+void FLayoutUV::FChartPacker::ScaleCharts( TArray< FMeshChart >& Charts, float UVScale )
 {
 	for( int32 i = 0; i < Charts.Num(); i++ )
 	{
@@ -766,7 +838,7 @@ void FLayoutUV::ScaleCharts( float UVScale )
 	Algo::IntroSort( Charts, FCompareCharts() );
 }
 
-bool FLayoutUV::PackCharts()
+bool FLayoutUV::FChartPacker::PackCharts(TArray< FMeshChart >& Charts, const TArray< FVector2D >& TexCoords, const TArray< uint32 >& SortedTris)
 {
 	uint32 RasterizeCycles = 0;
 	uint32 FindCycles = 0;
@@ -834,7 +906,7 @@ bool FLayoutUV::PackCharts()
 				else
 				{
 					int32 BeginRasterize = FPlatformTime::Cycles();
-					RasterizeChart( Chart, Rect.W, Rect.H );
+					RasterizeChart( Chart, TexCoords, SortedTris, Rect.W , Rect.H );
 					RasterizeCycles += FPlatformTime::Cycles() - BeginRasterize;
 				}
 
@@ -901,7 +973,7 @@ bool FLayoutUV::PackCharts()
 	return true;
 }
 
-void FLayoutUV::OrientChart( FMeshChart& Chart, int32 Orientation )
+void FLayoutUV::FChartPacker::OrientChart( FMeshChart& Chart, int32 Orientation )
 {
 	switch( Orientation )
 	{
@@ -959,8 +1031,8 @@ void FLayoutUV::OrientChart( FMeshChart& Chart, int32 Orientation )
 // Max of 2048x2048 due to precision
 // Dilate in 28.4 fixed point. Half pixel dilation is conservative rasterization.
 // Dilation same as Minkowski sum of triangle and square.
-template< typename TShader, int32 Dilate >
-void RasterizeTriangle( TShader& Shader, const FVector2D Points[3], int32 ScissorWidth, int32 ScissorHeight )
+template< int32 Dilate >
+void RasterizeTriangle( FAllocator2D& Shader, const FVector2D Points[3], int32 ScissorWidth, int32 ScissorHeight )
 {
 	const FVector2D HalfPixel( 0.5f, 0.5f );
 	FVector2D p0 = Points[0] - HalfPixel;
@@ -1031,13 +1103,13 @@ void RasterizeTriangle( TShader& Shader, const FVector2D Points[3], int32 Scisso
 
 			if( IsInside >= 0 )
 			{
-				Shader.Process( x, y );
+				Shader.SetBit( x, y );
 			}
 		}
 	}
 }
 
-void FLayoutUV::RasterizeChart( const FMeshChart& Chart, uint32 RectW, uint32 RectH )
+void FLayoutUV::FChartPacker::RasterizeChart( const FMeshChart& Chart, const TArray< FVector2D >& TexCoords, const TArray< uint32 >& SortedTris, uint32 RectW, uint32 RectH )
 {
 	// Bilinear footprint is -1 to 1 pixels. If packed geometrically, only a half pixel dilation
 	// would be needed to guarantee all charts were at least 1 pixel away, safe for bilinear filtering.
@@ -1049,13 +1121,13 @@ void FLayoutUV::RasterizeChart( const FMeshChart& Chart, uint32 RectW, uint32 Re
 	for( uint32 Tri = Chart.FirstTri; Tri < Chart.LastTri; Tri++ )
 	{
 		FVector2D Points[3];
-		for( int k = 0; k < 3; k++ )
+		for ( int k = 0; k < 3; k++ )
 		{
 			const FVector2D& UV = TexCoords[ 3 * SortedTris[ Tri ] + k ];
 			Points[k] = UV.X * Chart.PackingScaleU + UV.Y * Chart.PackingScaleV + Chart.PackingBias;
 		}
 
-		RasterizeTriangle< FAllocator2DShader, 16 >( ChartShader, Points, RectW, RectH );
+		RasterizeTriangle< 16 >( ChartRaster, Points, RectW, RectH );
 	}
 
 	if ( LayoutVersion >= ELightmapUVVersion::Segments )
@@ -1064,8 +1136,21 @@ void FLayoutUV::RasterizeChart( const FMeshChart& Chart, uint32 RectW, uint32 Re
 	}
 }
 
+bool FLayoutUV::FindBestPacking(uint32 InTextureResolution)
+{
+	FChartPacker Packer(MeshView, LayoutVersion, InTextureResolution);
+	bool bPackingFound = Packer.FindBestPacking(TexCoords, SortedTris, Charts);
+	PackedTextureResolution = bPackingFound ? InTextureResolution : 0;
+	return bPackingFound;
+}
+
 void FLayoutUV::CommitPackedUVs()
 {
+	if (PackedTextureResolution == 0)
+	{
+		return;
+	}
+
 	// Alloc new UV channel
 	MeshView.InitOutputTexcoords(TexCoords.Num());
 
@@ -1074,9 +1159,9 @@ void FLayoutUV::CommitPackedUVs()
 	{
 		FMeshChart& Chart = Charts[i];
 
-		Chart.PackingScaleU /= TextureResolution;
-		Chart.PackingScaleV /= TextureResolution;
-		Chart.PackingBias /= TextureResolution;
+		Chart.PackingScaleU /= PackedTextureResolution;
+		Chart.PackingScaleV /= PackedTextureResolution;
+		Chart.PackingBias /= PackedTextureResolution;
 
 		for( uint32 Tri = Chart.FirstTri; Tri < Chart.LastTri; Tri++ )
 		{
@@ -1091,28 +1176,28 @@ void FLayoutUV::CommitPackedUVs()
 	}
 }
 
-inline bool FLayoutUV::PositionsMatch( uint32 a, uint32 b ) const
+inline bool FLayoutUV::FChartFinder::PositionsMatch( uint32 a, uint32 b ) const
 {
 	return ( MeshView.GetPosition(a) - MeshView.GetPosition(b) ).IsNearlyZero( THRESH_POINTS_ARE_SAME );
 }
 
-inline bool FLayoutUV::NormalsMatch( uint32 a, uint32 b ) const
+inline bool FLayoutUV::FChartFinder::NormalsMatch( uint32 a, uint32 b ) const
 {
 	return ( MeshView.GetNormal(a) - MeshView.GetNormal(b) ).IsNearlyZero( THRESH_NORMALS_ARE_SAME );
 }
 
-inline bool FLayoutUV::UVsMatch( uint32 a, uint32 b ) const
+inline bool FLayoutUV::FChartFinder::UVsMatch( uint32 a, uint32 b ) const
 {
 	return ( MeshView.GetInputTexcoord(a) - MeshView.GetInputTexcoord(b) ).IsNearlyZero(UVLAYOUT_THRESH_UVS_ARE_SAME);
 }
 
-inline bool FLayoutUV::VertsMatch( uint32 a, uint32 b ) const
+inline bool FLayoutUV::FChartFinder::VertsMatch( uint32 a, uint32 b ) const
 {
 	return PositionsMatch( a, b ) && UVsMatch( a, b );
 }
 
 // Signed UV area
-inline float FLayoutUV::TriangleUVArea( uint32 Tri ) const
+inline float FLayoutUV::FChartFinder::TriangleUVArea( uint32 Tri ) const
 {
 	FVector2D UVs[3];
 	for( int k = 0; k < 3; k++ )
@@ -1125,7 +1210,7 @@ inline float FLayoutUV::TriangleUVArea( uint32 Tri ) const
 	return 0.5f * ( EdgeUV1.X * EdgeUV2.Y - EdgeUV1.Y * EdgeUV2.X );
 }
 
-inline void FLayoutUV::DisconnectChart( FMeshChart& Chart, uint32 Side )
+inline void FLayoutUV::FChartFinder::DisconnectChart( TArray< FMeshChart >& Charts, FMeshChart& Chart, uint32 Side )
 {
 	if( Chart.Join[ Side ] != -1 )
 	{
@@ -1134,7 +1219,7 @@ inline void FLayoutUV::DisconnectChart( FMeshChart& Chart, uint32 Side )
 	}
 }
 
-inline float FLayoutUV::GetUVEqualityThreshold() const
+inline float FLayoutUV::FChartFinder::GetUVEqualityThreshold() const
 {
 	return LayoutVersion >= ELightmapUVVersion::SmallChartPacking ? NEW_UVS_ARE_SAME : LEGACY_UVS_ARE_SAME;
 }
diff --git a/Engine/Source/Developer/MeshUtilitiesCommon/Public/Allocator2D.h b/Engine/Source/Developer/MeshUtilitiesCommon/Public/Allocator2D.h
--- a/Engine/Source/Developer/MeshUtilitiesCommon/Public/Allocator2D.h
+++ b/Engine/Source/Developer/MeshUtilitiesCommon/Public/Allocator2D.h
@@ -37,11 +37,7 @@ public:
 	};
 
 public:
-				FAllocator2D( uint32 Width, uint32 Height );
-				FAllocator2D( const FAllocator2D& Other );
-				~FAllocator2D();
-	
-	FAllocator2D& operator=( const FAllocator2D& Other );
+	FAllocator2D( uint32 Width, uint32 Height );
 
 	// Must clear before using
 	void		Clear();
@@ -73,7 +69,7 @@ protected:
 	void		AddUsedSegment( FRow& Row, uint32 StartPos, uint32 Length );
 
 private:
-	uint64*		Bits;
+	TArray< uint64 > Bits;
 
 	uint32		Width;
 	uint32		Height;
@@ -83,20 +79,6 @@ private:
 	int32		LastRowFail;
 };
 
-struct FAllocator2DShader
-{
-	FAllocator2D*	Allocator2D;
-
-	FAllocator2DShader( FAllocator2D* InAllocator2D )
-	: Allocator2D( InAllocator2D )
-	{}
-
-	FORCEINLINE void Process( uint32 x, uint32 y )
-	{
-		Allocator2D->SetBit( x, y );
-	}
-};
-
 // Returns non-zero if set
 FORCEINLINE uint64 FAllocator2D::GetBit( uint32 x, uint32 y ) const
 {
diff --git a/Engine/Source/Developer/MeshUtilitiesCommon/Public/LayoutUV.h b/Engine/Source/Developer/MeshUtilitiesCommon/Public/LayoutUV.h
--- a/Engine/Source/Developer/MeshUtilitiesCommon/Public/LayoutUV.h
+++ b/Engine/Source/Developer/MeshUtilitiesCommon/Public/LayoutUV.h
@@ -53,44 +53,21 @@ public:
 		virtual void        SetOutputTexcoord(uint32 Index, const FVector2D& Value) = 0;
 	};
 
-				FLayoutUV( IMeshView& InMeshView, uint32 InTextureResolution );
-
-	int32		FindCharts( const FOverlappingCorners& OverlappingCorners );
-	bool		FindBestPacking();
-	void		CommitPackedUVs();
-
-	void		SetVersion( ELightmapUVVersion Version ) { LayoutVersion = Version; }
+	FLayoutUV( IMeshView& InMeshView );
+	void SetVersion( ELightmapUVVersion Version ) { LayoutVersion = Version; }
+	int32 FindCharts( const FOverlappingCorners& OverlappingCorners );
+	bool FindBestPacking( uint32 InTextureResolution );
+	void CommitPackedUVs();
 
 private:
-	bool		PositionsMatch( uint32 a, uint32 b ) const;
-	bool		NormalsMatch( uint32 a, uint32 b ) const;
-	bool		UVsMatch( uint32 a, uint32 b ) const;
-	bool		VertsMatch( uint32 a, uint32 b ) const;
-	float		TriangleUVArea( uint32 Tri ) const;
-	void		DisconnectChart( FMeshChart& Chart, uint32 Side );
-
-	void		ScaleCharts( float UVScale );
-	bool		PackCharts();
-	void		OrientChart( FMeshChart& Chart, int32 Orientation );
-	void		RasterizeChart( const FMeshChart& Chart, uint32 RectW, uint32 RectH );
-
-	float		GetUVEqualityThreshold() const;
-
-	IMeshView&	MeshView;
-	uint32		TextureResolution;
-	
-	TArray< FVector2D >		TexCoords;
-	TArray< uint32 >		SortedTris;
-	TArray< FMeshChart >	Charts;
-	float					TotalUVArea;
-	float					MaxChartSize;
-
-	FAllocator2D		LayoutRaster;
-	FAllocator2D		ChartRaster;
-	FAllocator2D		BestChartRaster;
-	FAllocator2DShader	ChartShader;
+	IMeshView& MeshView;
+	ELightmapUVVersion LayoutVersion;
 
-	ELightmapUVVersion	LayoutVersion;
+	TArray< FVector2D > TexCoords;
+	TArray< uint32 > SortedTris;
+	TArray< FMeshChart > Charts;
+	uint32 PackedTextureResolution;
 
-	int32				NextMeshChartId;
+	struct FChartFinder;
+	struct FChartPacker;
 };
