fix build warning for [Compile UE4Editor Linux]

diff --git a/Engine/Source/Developer/MeshUtilitiesCommon/Private/LayoutUV.cpp b/Engine/Source/Developer/MeshUtilitiesCommon/Private/LayoutUV.cpp
--- a/Engine/Source/Developer/MeshUtilitiesCommon/Private/LayoutUV.cpp
+++ b/Engine/Source/Developer/MeshUtilitiesCommon/Private/LayoutUV.cpp
@@ -86,7 +86,7 @@ FLayoutUV::FChartPacker::FChartPacker(IMeshView& InMeshView, ELightmapUVVersion
 int32 FLayoutUV::FindCharts(const FOverlappingCorners& OverlappingCorners)
 {
 	FChartFinder Finder(MeshView, LayoutVersion);
-	return Finder.FindCharts(OverlappingCorners, TexCoords, SortedTris, Charts);
+	return Finder.FindCharts(OverlappingCorners, MeshTexCoords, MeshSortedTris, MeshCharts);
 }
 
 FLayoutUV::FChartFinder::FChartFinder(IMeshView& InMeshView, ELightmapUVVersion InLayoutVersion)
@@ -1139,7 +1139,7 @@ void FLayoutUV::FChartPacker::RasterizeChart( const FMeshChart& Chart, const TAr
 bool FLayoutUV::FindBestPacking(uint32 InTextureResolution)
 {
 	FChartPacker Packer(MeshView, LayoutVersion, InTextureResolution);
-	bool bPackingFound = Packer.FindBestPacking(TexCoords, SortedTris, Charts);
+	bool bPackingFound = Packer.FindBestPacking(MeshTexCoords, MeshSortedTris, MeshCharts);
 	PackedTextureResolution = bPackingFound ? InTextureResolution : 0;
 	return bPackingFound;
 }
@@ -1152,12 +1152,12 @@ void FLayoutUV::CommitPackedUVs()
 	}
 
 	// Alloc new UV channel
-	MeshView.InitOutputTexcoords(TexCoords.Num());
+	MeshView.InitOutputTexcoords(MeshTexCoords.Num());
 
 	// Commit chart UVs
-	for( int32 i = 0; i < Charts.Num(); i++ )
+	for( int32 i = 0; i < MeshCharts.Num(); i++ )
 	{
-		FMeshChart& Chart = Charts[i];
+		FMeshChart& Chart = MeshCharts[i];
 
 		Chart.PackingScaleU /= PackedTextureResolution;
 		Chart.PackingScaleV /= PackedTextureResolution;
@@ -1167,8 +1167,8 @@ void FLayoutUV::CommitPackedUVs()
 		{
 			for( int k = 0; k < 3; k++ )
 			{
-				uint32 Index = 3 * SortedTris[ Tri ] + k;
-				const FVector2D& UV = TexCoords[ Index ];
+				uint32 Index = 3 * MeshSortedTris[ Tri ] + k;
+				const FVector2D& UV = MeshTexCoords[ Index ];
 				FVector2D TransformedUV = UV.X * Chart.PackingScaleU + UV.Y * Chart.PackingScaleV + Chart.PackingBias;
 				MeshView.SetOutputTexcoord(Index, TransformedUV);
 			}
diff --git a/Engine/Source/Developer/MeshUtilitiesCommon/Public/LayoutUV.h b/Engine/Source/Developer/MeshUtilitiesCommon/Public/LayoutUV.h
--- a/Engine/Source/Developer/MeshUtilitiesCommon/Public/LayoutUV.h
+++ b/Engine/Source/Developer/MeshUtilitiesCommon/Public/LayoutUV.h
@@ -63,9 +63,9 @@ private:
 	IMeshView& MeshView;
 	ELightmapUVVersion LayoutVersion;
 
-	TArray< FVector2D > TexCoords;
-	TArray< uint32 > SortedTris;
-	TArray< FMeshChart > Charts;
+	TArray< FVector2D > MeshTexCoords;
+	TArray< uint32 > MeshSortedTris;
+	TArray< FMeshChart > MeshCharts;
 	uint32 PackedTextureResolution;
 
 	struct FChartFinder;
