Improve performance of UV island packing
  - Test rows from longest to shortest to get rid of one inner loop in the hot path
  - Add lookup table to skip another free segment inner loop in the hot path
  - Added island hashing to skip directly to last known best position
  - Split loops, simplify code and add comments for clarity
  - Reduce data copying and clearing to working area only
  - Multithreaded initial scale search when complexity is high enough
  - Added a segment grid along the Y axis for more efficient search of taller island
  - Added ability to log packing stats per island in Insights to understand algorithmic behavior
  - Now able to pack 300k islands on a 2048x2048 texture in under 60s using 4 cores

  - Improvements made while maintaining backward compatibility (same output as before)
     - 52m16s -> 80s for DATASET-0002a (foliage)
     - 148s -> 105s for DATASET-0008a
     - 333s -> 8.5s for DATASET-0020a (house)  using 2048x2048 MinResolution
  - Improvements made while activating segments in 2D including packing efficiency improvements
     - 52m16s -> 14s (efficiency 95.52% -> 97.68%) for DATASET-0002a (foliage)
     - 148s -> 105s (efficiency 63.3% -> 63.3%) for DATASET-0008a
     - 333s -> 6s  (efficiency 93.5% -> 98.57%) for DATASET-0020 (house) using 2048x2048 MinResolution

From dev-enterprise[REDACTED_HANDLE]

diff --git a/Engine/Plugins/Enterprise/DatasmithImporter/Source/DatasmithImporter/Private/DatasmithImportFactory.cpp b/Engine/Plugins/Enterprise/DatasmithImporter/Source/DatasmithImporter/Private/DatasmithImportFactory.cpp
--- a/Engine/Plugins/Enterprise/DatasmithImporter/Source/DatasmithImporter/Private/DatasmithImportFactory.cpp
+++ b/Engine/Plugins/Enterprise/DatasmithImporter/Source/DatasmithImporter/Private/DatasmithImportFactory.cpp
@@ -11,6 +11,7 @@
 #include "DatasmithSceneFactory.h"
 #include "DatasmithStaticMeshImporter.h"
 #include "IDatasmithSceneElements.h"
+#include "LayoutUV.h"
 #include "Translators/DatasmithTranslatableSource.h"
 #include "Translators/DatasmithTranslatorManager.h"
 #include "Utility/DatasmithImporterUtils.h"
@@ -292,6 +293,8 @@ namespace DatasmithImportFactoryImpl
 
 	static void ReportImportStats(const FDatasmithImportContext& ImportContext, uint64 StartTime)
 	{
+		FLayoutUV::LogStats();
+
 		// Log time spent to import incoming file in minutes and seconds
 		double ElapsedSeconds = FPlatformTime::ToSeconds64(FPlatformTime::Cycles64() - StartTime);
 
@@ -469,6 +472,8 @@ bool UDatasmithImportFactory::Import( FDatasmithImportContext& ImportContext )
 		bShowOptions = false;
 	}
 
+	FLayoutUV::ResetStats();
+
 	if ( !ImportContext.bIsAReimport )
 	{
 		GEditor->GetEditorSubsystem<UImportSubsystem>()->BroadcastAssetPreImport(this, UDatasmithScene::StaticClass(), nullptr, ImportContext.Scene->GetName(), nullptr);
diff --git a/Engine/Source/Developer/MeshUtilitiesCommon/Private/Allocator2D.cpp b/Engine/Source/Developer/MeshUtilitiesCommon/Private/Allocator2D.cpp
--- a/Engine/Source/Developer/MeshUtilitiesCommon/Private/Allocator2D.cpp
+++ b/Engine/Source/Developer/MeshUtilitiesCommon/Private/Allocator2D.cpp
@@ -2,21 +2,38 @@
 
 #include "Allocator2D.h"
 
-FAllocator2D::FAllocator2D( uint32 InWidth, uint32 InHeight )
-	: Width( InWidth )
+#include "Misc/SecureHash.h"
+#include "ProfilingDebugging/MiscTrace.h"
+
+FAllocator2D::FAllocator2D( FAllocator2D::EMode InMode, uint32 InWidth, uint32 InHeight, ELightmapUVVersion InLayoutVersion)
+	: Mode( InMode )
+	, Width( InWidth )
 	, Height( InHeight )
-	, LastRowFail( -1 )
+	, Pitch( ( InWidth + 63 ) / 64 )
+	, LayoutVersion(InLayoutVersion)
+	, RasterWidth(0)
+	, RasterHeight(0)
 {
-	Pitch = ( Width + 63 ) / 64;
+	check(Width  <= MAX_uint16);
+	check(Height <= MAX_uint16);
 
-	Rows.Reserve( Height );
+	Bits.SetNumZeroed(Pitch * Height + 1); // alloc +1 to avoid buffer overrun
 
-	for ( uint32 Y = 0; Y < Height; ++Y )
+	Rows.SetNum( Height );
+	SortedRowsIndex.SetNum(Height);
+	for (uint32 Index = 0; Index < Height; ++Index)
 	{
-		FRow Row;
-		Row.Index = Y;
+		SortedRowsIndex[Index] = (uint16)Index;
+	}
 
-		Rows.Add( Row );
+	if (LayoutVersion >= ELightmapUVVersion::Segments2D)
+	{
+		Columns.SetNum( Width );
+		SortedColumnsIndex.SetNum(Width);
+		for (uint32 Index = 0; Index < Width; ++Index)
+		{
+			SortedColumnsIndex[Index] = (uint16)Index;
+		}
 	}
 
 	Clear();
@@ -25,8 +42,83 @@ FAllocator2D::FAllocator2D( uint32 InWidth, uint32 InHeight )
 void FAllocator2D::Clear()
 {
 	InitSegments();
-	Bits.Reset();
-	Bits.SetNumZeroed(Pitch * Height + 1); // alloc +1 to avoid buffer overrun
+
+	// Only clear the section we used to reduce memory traffic
+	FPlatformMemory::Memzero(Bits.GetData(), Pitch * RasterHeight * Bits.GetTypeSize());
+
+	RasterWidth  = 0;
+	RasterHeight = 0;
+}
+
+void FAllocator2D::CopyRuns(TArray<FRun>& Runs, const TArray<FRun>& OtherRuns, int32 MaxSize)
+{
+	Runs.SetNum(OtherRuns.Num(), false);
+
+	MaxSize = FMath::Min(Runs.Num(), MaxSize);
+
+	FRun* RunsPtr = Runs.GetData();
+	const FRun* OtherRunsPtr = OtherRuns.GetData();
+	for (int32 Index = 0; Index < MaxSize; ++Index, ++RunsPtr, ++OtherRunsPtr)
+	{
+		*RunsPtr = *OtherRunsPtr;
+	}
+}
+
+// Optimized version to alloc and copy as less data possible.
+FAllocator2D& FAllocator2D::operator = (const FAllocator2D & Other)
+{
+	if (&Other == this)
+	{
+		return *this;
+	}
+
+	check(Mode == Other.Mode);
+	check(Width == Other.Width);
+	check(Height == Other.Height);
+	check(Pitch == Other.Pitch);
+
+	FPlatformMemory::Memcpy(Bits.GetData(), Other.Bits.GetData(), Pitch * Other.RasterHeight * Bits.GetTypeSize());
+
+	// Zero out anything that was used and not overwritten by the copy
+	if (RasterHeight > Other.RasterHeight)
+	{
+		FPlatformMemory::Memzero(Bits.GetData() + Pitch * Other.RasterHeight, Pitch * (RasterHeight - Other.RasterHeight) * Bits.GetTypeSize());
+	}
+
+	RasterWidth  = Other.RasterWidth;
+	RasterHeight = Other.RasterHeight;
+
+	CopyRuns(Rows, Other.Rows, Other.RasterHeight);
+	SortedRowsIndex = Other.SortedRowsIndex;
+
+	if (LayoutVersion >= ELightmapUVVersion::Segments2D)
+	{
+		CopyRuns(Columns, Other.Columns, Other.RasterWidth);
+		SortedColumnsIndex = Other.SortedColumnsIndex;
+	}
+
+	return *this;
+}
+
+FMD5Hash FAllocator2D::GetRasterMD5() const
+{
+	FMD5 MD5;
+
+	const FRun* RowPtr = Rows.GetData();
+	const FRun* RowPtrEnd = RowPtr + RasterHeight;
+
+	for (; RowPtr < RowPtrEnd; ++RowPtr)
+	{
+		uint32 SegmentsNum = RowPtr->Segments.Num();
+		
+		// It is important for empty rows to also have a footprint in the MD5
+		MD5.Update((const uint8*)&SegmentsNum, sizeof(SegmentsNum));
+		MD5.Update((const uint8*)RowPtr->Segments.GetData(), SegmentsNum * RowPtr->Segments.GetTypeSize() );
+	}
+
+	FMD5Hash MD5Hash;
+	MD5Hash.Set(MD5);
+	return MD5Hash;
 }
 
 bool FAllocator2D::Find( FRect& Rect )
@@ -47,6 +139,60 @@ bool FAllocator2D::Find( FRect& Rect )
 	return false;
 }
 
+FString FAllocator2D::ToString() const
+{
+	TArray<TCHAR> Text;
+	Text.SetNum(Width+2);
+
+	// Truncate output to last meaningful row
+	int32 lastRow = 1;
+	for (const FRun& Row : Rows)
+	{
+		if ((Mode == EMode::FreeSegments && Row.Segments.Num() > 1) ||
+			(Mode == EMode::UsedSegments && Row.Segments.Num() != 0))
+		{
+			lastRow++;
+		}
+	}
+
+	FString Output;
+	Output.Append(TEXT("BEGIN -----------------------\n"));
+
+	const TCHAR FillChar = Mode == EMode::FreeSegments ? TEXT('-') : TEXT(' ');
+	const TCHAR UsedChar = Mode == EMode::FreeSegments ? TEXT(' ') : TEXT('+');
+
+	for (const FRun& Row : Rows)
+	{
+		for (uint32 x = 0; x < Width; ++x)
+		{
+			Text[x] = FillChar;
+		}
+
+		for (const FSegment& Segment : Row.Segments)
+		{
+			for (uint32 x = Segment.StartPos; x < Segment.StartPos + Segment.Length; ++x)
+			{
+				Text[x] = UsedChar;
+			}
+		}
+
+		Text[Text.Num()-2] = TEXT('\n');
+		Text[Text.Num()-1] = TEXT('\0');
+
+		Output.Append(Text.GetData());
+
+		if (--lastRow == 0)
+		{
+			break;
+		}
+	}
+
+	Output.Appendf(TEXT("MD5 %s\n"), *LexToString(GetRasterMD5()));
+	Output.Append(TEXT("END -----------------------\n"));
+
+	return Output;
+}
+
 bool FAllocator2D::FindBitByBit( FRect& Rect, const FAllocator2D& Other )
 {
 	FRect TestRect = Rect;
@@ -65,27 +211,72 @@ bool FAllocator2D::FindBitByBit( FRect& Rect, const FAllocator2D& Other )
 	return false;
 }
 
-bool FAllocator2D::FindWithSegments( FRect& Rect, FRect BestRect, const FAllocator2D& Other )
+bool FAllocator2D::FindWithSegments( FRect& Rect, const FRect& BestRect, const FAllocator2D& Other ) const
 {
-	LastRowFail = -1;
 	FRect TestRect = Rect;
 
-	for( TestRect.Y = 0; TestRect.Y <= Height - TestRect.H; ++TestRect.Y )
-	{
-		uint32 FailedLength = 1;
+	const uint32 MaxWidth    = Width - TestRect.W;
+	const uint32 MaxHeight   = Height - TestRect.H;
+	const uint32 BestArea    = BestRect.X + BestRect.Y * Width;
 
-		for( TestRect.X = 0; TestRect.X <= Width - TestRect.W; TestRect.X += FailedLength )
+	// For charts that are longer on the Y axis, its a lot faster to use
+	// another grid optimized for searching vertically because it will reduce
+	// iteration count tremendously by improving strides we make for each iteration.
+	// We need a version check to activate this feature because it will favor
+	// placement along the Y axis leading to (still efficient) but different results.
+	if (LayoutVersion >= ELightmapUVVersion::Segments2D && TestRect.H > TestRect.W)
+	{
+		// Perfect fit algorithm will try every position to fit our chart
+		for( ; TestRect.X <= MaxWidth; ++TestRect.X)
 		{
-			if( TestRect.X + TestRect.Y * Height >= BestRect.X + BestRect.Y * Height )
+			Stats.FindWithSegmentsIterationsX++;
+
+			// This represents the length we know for sure we won't fit in so we can skip over
+			uint32 FailedLength = 1;
+			for( TestRect.Y = 0; TestRect.Y <= MaxHeight; TestRect.Y += FailedLength)
 			{
-				// We've moved past the BestRect
-				return false;
+				Stats.FindWithSegmentsIterationsY++;
+
+				// Do not bother continue if another orientation had a better result
+				if( TestRect.X + TestRect.Y * Width >= BestArea)
+				{
+					Stats.FindWithSegmentsMovedPastPreviousBest++;
+					return false;
+				}
+
+				if( TestAllColumns( TestRect, Other, FailedLength ) )
+				{
+					Rect = TestRect;
+					return true;
+				}
 			}
+		}
+	}
+	else
+	{
+		// Perfect fit algorithm will try every position to fit our chart
+		for( ; TestRect.Y <= MaxHeight; ++TestRect.Y)
+		{
+			Stats.FindWithSegmentsIterationsY++;
 
-			if( TestAllRows( TestRect, Other, FailedLength ) )
+			// This represents the length we know for sure we won't fit in so we can skip over
+			uint32 FailedLength = 1;
+			for( TestRect.X = 0; TestRect.X <= MaxWidth; TestRect.X += FailedLength)
 			{
-				Rect = TestRect;
-				return true;
+				Stats.FindWithSegmentsIterationsX++;
+
+				// Do not bother continue if another orientation had a better result
+				if( TestRect.X + TestRect.Y * Width >= BestArea)
+				{
+					Stats.FindWithSegmentsMovedPastPreviousBest++;
+					return false;
+				}
+
+				if( TestAllRows( TestRect, Other, FailedLength ) )
+				{
+					Rect = TestRect;
+					return true;
+				}
 			}
 		}
 	}
@@ -93,6 +284,25 @@ bool FAllocator2D::FindWithSegments( FRect& Rect, FRect BestRect, const FAllocat
 	return false;
 }
 
+uint32 FAllocator2D::GetUsedTexels() const
+{
+	uint32 Texels = 0;
+	for (const FRun& Row : Rows)
+	{
+		for (const FSegment& Segment : Row.Segments)
+		{
+			Texels += Segment.Length;
+		}
+	}
+
+	if (Mode == EMode::FreeSegments)
+	{
+		Texels = Width * Height - Texels;
+	}
+
+	return Texels;
+}
+
 void FAllocator2D::Alloc( FRect Rect )
 {
 	for( uint32 y = Rect.Y; y < Rect.Y + Rect.H; y++ )
@@ -120,99 +330,125 @@ void FAllocator2D::Alloc( FRect Rect, const FAllocator2D& Other )
 	MergeSegments( Rect, Other );
 }
 
-bool FAllocator2D::TestAllRows( FRect Rect, const FAllocator2D& Other, uint32& FailedLength )
+bool FAllocator2D::TestAllRows( const FRect& Rect, const FAllocator2D& Other, uint32& OutFailedLength ) const
 {
-	FailedLength = 0.f;
-	bool bSuccess = true;
-
-	if ( LastRowFail > 0 )
+	for ( uint32 Index = 0; Index < Other.RasterHeight; ++Index)
 	{
-		const FRow& ThisRow = Rows[ Rect.Y + LastRowFail ];
-		const FRow& OtherRow = Other.Rows[ LastRowFail ];
+		Stats.TestAllRowsIterationsY++;
+
+		// Longest rows first have a higher chance of giving us a big stride when it fails
+		const uint16 y = Other.SortedRowsIndex.GetData()[Index];
 
-		if ( !TestRow( ThisRow, OtherRow, Rect, FailedLength ) )
+		const FRun& ThisRow  = Rows.GetData()[ Rect.Y + y ];
+		const FRun& OtherRow = Other.Rows.GetData()[ y ];
+
+		if (!TestOneRun(ThisRow, OtherRow, Rect.X, Rect.W, Width, OutFailedLength))
 		{
 			return false;
 		}
 	}
 
-	LastRowFail = -1;
+	return true;
+}
 
-	for ( uint32 y = 0; y < Rect.H; ++y )
+bool FAllocator2D::TestAllColumns( const FRect& Rect, const FAllocator2D& Other, uint32& OutFailedLength ) const
+{
+	for ( uint32 Index = 0; Index < Other.RasterWidth; ++Index)
 	{
-		const FRow& ThisRow = Rows[ Rect.Y + y ];
-		const FRow& OtherRow = Other.Rows[ y ];
+		// Longest columns first have a higher chance of giving us a big stride when it fails
+		const uint16 x = Other.SortedColumnsIndex.GetData()[Index];
 
-		uint32 CurrentFailedLength = 0;
-		bool bThisRowSuccess = TestRow( ThisRow, OtherRow, Rect, CurrentFailedLength );
-		bSuccess = bThisRowSuccess && bSuccess;
-
-		if ( !bThisRowSuccess && CurrentFailedLength > FailedLength )
-		{
-			LastRowFail = y;
-			FailedLength = CurrentFailedLength;
-		}
+		const FRun& ThisColumn  = Columns.GetData()[ Rect.X + x ];
+		const FRun& OtherColumn = Other.Columns.GetData()[ x ];
 
-		if ( FailedLength >= Width )
+		if (!TestOneRun(ThisColumn, OtherColumn, Rect.Y, Rect.H, Height, OutFailedLength))
 		{
 			return false;
 		}
 	}
 
-	return bSuccess;
+	return true;
 }
 
-bool FAllocator2D::TestRow( const FRow& ThisRow, const FRow& OtherRow, FRect Rect, uint32& FailedLength )
+bool FAllocator2D::TestOneRun(const FRun& ThisRun, const FRun& OtherRun, uint32 RectOffset, uint32 RectLength, uint32 PrimaryResolution, uint32& OutFailedLength) const
 {
-	if ( ThisRow.LongestSegment < OtherRow.LongestSegment )
+	// Early out if we can't fit the longest used segment into our longest free segment
+	if ( ThisRun.LongestSegment < OtherRun.LongestSegment )
 	{
-		FailedLength = Width;
+		OutFailedLength = PrimaryResolution;
 		return false;
 	}
 
-	uint32 StartFreeSegmentIndex = 0;
-
-	for ( const FSegment& OtherUsedSegment : OtherRow.UsedSegments )
+	// Threshold to perform a linear scan until the cache miss induced by the lookup is worth it
+	const int32     LoopkupThreshold  = 16;
+	const FSegment* FreeSegments      = ThisRun.Segments.GetData();
+	const FSegment* FreeSegmentsEnd   = ThisRun.Segments.GetData() + ThisRun.Segments.Num();
+	const uint16*   FreeSegmentLookup = ThisRun.Segments.Num() > LoopkupThreshold ? ThisRun.FreeSegmentsLookup.GetData() : nullptr;
+	const FSegment* UsedSegments      = OtherRun.Segments.GetData();
+	const FSegment* UsedSegmentsEnd   = OtherRun.Segments.GetData() + OtherRun.Segments.Num();
+
+	// Running pointer for used segments
+	const FSegment* UsedSegment = UsedSegments;
+	for ( ; UsedSegment < UsedSegmentsEnd; ++UsedSegment)
 	{
-		if ( OtherUsedSegment.StartPos >= Rect.W )
+		// #TODO maintain this for backward compatibility but not sure if this could ever happen
+		// and if the result is valid if we break instead of returning false when a used segment doesn't fit
+		if ( UsedSegment->StartPos >= RectLength )
 		{
 			break;
 		}
 
-		bool bFoundSpaceForSegment = false;
-		bool bFoundFutureSpaceForSegment = false;
+		const uint32 StartPos = RectOffset + UsedSegment->StartPos;
+		const uint32 EndPos   = RectOffset + FMath::Min( UsedSegment->StartPos + UsedSegment->Length, RectLength );
 
-		const uint32 StartPos = Rect.X + OtherUsedSegment.StartPos;
-		const uint32 EndPos = Rect.X + FMath::Min( OtherUsedSegment.StartPos + OtherUsedSegment.Length, Rect.W );
+		// Running pointer for free segments
+		const FSegment* FreeSegment = FreeSegments;
 
-		for ( int32 i = StartFreeSegmentIndex; i < ThisRow.FreeSegments.Num(); ++i )
+		// If lookup exist, advance pointer to the exact spot instead of scanning
+		if (FreeSegmentLookup)
 		{
-			const FSegment& ThisFreeSegment = ThisRow.FreeSegments[i];
+			FreeSegment += FreeSegmentLookup[StartPos];
+			Stats.FreeSegmentLookupCount++;
+		}
 
-			if ( StartPos >= ThisFreeSegment.StartPos &&
-				EndPos <= ThisFreeSegment.StartPos + ThisFreeSegment.Length )
+		// Scan the search range to see if the current used segment fits in
+		// we use a for loop to support scanning from the beginning if no lookup exists
+		bool bFoundSpaceForSegment = false;
+		for (; FreeSegment < FreeSegmentsEnd && StartPos >= FreeSegment->StartPos; ++FreeSegment)
+		{
+			Stats.FreeSegmentRangeIterations++;
+
+			// Check if there's a free segment that can fit the used segment
+			if (EndPos <= FreeSegment->StartPos + FreeSegment->Length)
 			{
-				StartFreeSegmentIndex = i;
 				bFoundSpaceForSegment = true;
 				break;
 			}
-			else if ( !bFoundFutureSpaceForSegment &&
-				StartPos < ThisFreeSegment.StartPos &&
-				OtherUsedSegment.Length <= ThisFreeSegment.Length )
-			{
-				bFoundFutureSpaceForSegment = true;
-				FailedLength = ThisFreeSegment.StartPos - StartPos;
-				break;
-			}
 		}
 
-		if ( !bFoundSpaceForSegment )
+		// If nothing was found in the search range, scan other free segments to hint our
+		// caller of how much farther it should advance for next test
+		if (!bFoundSpaceForSegment)
 		{
-			if ( !bFoundFutureSpaceForSegment )
+			for (; FreeSegment < FreeSegmentsEnd; ++FreeSegment)
 			{
-				FailedLength = Width;
+				Stats.FreeSegmentFutureIterations++;
+				if (UsedSegment->Length <= FreeSegment->Length)
+				{
+					OutFailedLength = FreeSegment->StartPos - StartPos;
+
+					Stats.FreeSegmentFutureHit++;
+					Stats.FreeSegmentFutureHitStep += OutFailedLength;
+					return false;
+				}
 			}
 
+			// We couldn't find any free segment big enough
+			// set failed length so the whole run will be skipped
+			OutFailedLength = PrimaryResolution;
+
+			Stats.FreeSegmentFutureMiss++;
+			Stats.FreeSegmentFutureMissStep += OutFailedLength;
 			return false;
 		}
 	}
@@ -220,30 +456,15 @@ bool FAllocator2D::TestRow( const FRow& ThisRow, const FRow& OtherRow, FRect Rec
 	return true;
 }
 
-void FAllocator2D::FlipX( FRect Rect, ELightmapUVVersion LayoutVersion )
+void FAllocator2D::FlipX( const FRect& Rect )
 {
 	// If we have empty padding around the Rect, keep it there
-	uint32 MaxY = Rect.H - 1;
-	for ( int32 Y = Rect.H - 1; Y > 0; --Y )
-	{
-		if ( Rows[ Y ].UsedSegments.Num() > 0 )
-		{
-			break;
-		}
-
-		--MaxY;
-	}
-
 	uint32 MaxX = 0;
+	uint32 MaxY = RasterHeight - 1;
+
 	if ( LayoutVersion >= ELightmapUVVersion::Allocator2DFlipFix )
 	{
-		for ( uint32 Y = 0; Y <= MaxY; ++Y )
-		{
-			for ( const FSegment& UsedSegment : Rows[ Y ].UsedSegments )
-			{
-				MaxX = FMath::Max( MaxX, UsedSegment.StartPos + UsedSegment.Length - 1 );
-			}
-		}
+		MaxX = RasterWidth - 1;
 	}
 	else
 	{
@@ -282,22 +503,14 @@ void FAllocator2D::FlipX( FRect Rect, ELightmapUVVersion LayoutVersion )
 	CreateUsedSegments();
 }
 
-void FAllocator2D::FlipY( FRect Rect )
+void FAllocator2D::FlipY( const FRect& Rect )
 {
 	uint32 MinY = 0;
-
-	uint32 MaxY = Rect.H - 1;
-	for ( int32 Y = Rect.H - 1; Y > 0; --Y )
-	{
-		if ( Rows[ Y ].UsedSegments.Num() > 0 )
-			break;
-
-		--MaxY;
-	}
+	uint32 MaxY = RasterHeight - 1;
 
 	for ( uint32 LowY = 0; LowY < ( MaxY + 1 ) / 2; ++LowY )
 	{
-		for ( uint32 X = 0; X < Rect.W; ++X )
+		for ( uint32 X = 0; X < RasterWidth; ++X )
 		{
 			uint32 HighY = MaxY - LowY;
 
@@ -324,48 +537,72 @@ void FAllocator2D::FlipY( FRect Rect )
 		}
 	}
 
-	for ( uint32 LowY = MinY; LowY < ( MaxY + 1 ) / 2; ++LowY )
+	CreateUsedSegments();
+}
+
+void FAllocator2D::InitRuns(TArray<FRun>& Runs, uint32 PrimaryResolution, uint32 PerpendicularRasterSize)
+{
+	if (Mode == EMode::FreeSegments)
 	{
-		const int32 HighY = MaxY - LowY;
+		FSegment FreeSegment;
+		FreeSegment.StartPos = 0;
+		FreeSegment.Length = PrimaryResolution;
 
-		Rows.Swap( LowY, HighY );
-		Rows[ LowY ].Index = LowY;
-		Rows[ HighY ].Index = HighY;
+		for ( FRun& Run : Runs )
+		{
+			Run.Segments.Reset();
+			Run.Segments.Add( FreeSegment );
+			Run.LongestSegment = PrimaryResolution;
+			Run.FreeSegmentsLookup.Reset();
+		}
+	}
+	else
+	{
+		// This is called inside a hot loop, get rid of RangeCheck.
+		FRun* Run = Runs.GetData();
+		for ( uint32 Index = 0; Index < PerpendicularRasterSize; ++Index )
+		{
+			Run[Index].Segments.Reset();
+			Run[Index].LongestSegment = 0;
+		}
 	}
 }
 
 void FAllocator2D::InitSegments()
 {
-	FSegment FreeSegment;
-	FreeSegment.StartPos = 0;
-	FreeSegment.Length = Width;
+	InitRuns(Rows, Width, RasterHeight);
 
-	for ( FRow& Row : Rows )
+	if (LayoutVersion >= ELightmapUVVersion::Segments2D)
 	{
-		Row.FreeSegments.Reset();
-		Row.FreeSegments.Add( FreeSegment );
-		Row.LongestSegment = FreeSegment.Length;
-
-		Row.UsedSegments.Reset();
-
+		InitRuns(Columns, Height, RasterWidth);
 	}
 }
 
 void FAllocator2D::CreateUsedSegments()
 {
-	for ( uint32 y = 0; y < Height; ++y )
+	check(Mode == EMode::UsedSegments);
+
+	SortedRowsIndex.SetNum(RasterHeight, false);
+
+	uint64* BitsData = Bits.GetData();
+
+	// Create segments along the X axis for each rows
+	for ( uint32 y = 0; y < RasterHeight; ++y )
 	{
-		FRow& CurrentRow = Rows[y];
+		SortedRowsIndex[y] = y;
+
+		FRun& CurrentRow = Rows[y];
 		CurrentRow.LongestSegment = 0;
-		CurrentRow.UsedSegments.Reset();
+		CurrentRow.Segments.Reset();
 
 		int32 FirstUsedX = -1;
-		for ( uint32 k = 0; k < Pitch; ++k )
+		for ( uint32 k = 0; k < Pitch; ++k, ++BitsData )
 		{
-			uint32 x = k * 64;
+			const uint32 x = k * 64;
+			const uint64 BitsValue = *BitsData;
 
 			// If all bits are set
-			if ( Bits[k + y * Pitch] == ~0ull )
+			if ( BitsValue == ~0ull )
 			{
 				if ( FirstUsedX < 0 )
 				{
@@ -379,7 +616,7 @@ void FAllocator2D::CreateUsedSegments()
 				}
 			}
 			// No bits are set
-			else if ( Bits[k + y * Pitch] == 0ull )
+			else if ( BitsValue == 0ull )
 			{
 				if ( FirstUsedX >= 0 )
 				{
@@ -394,7 +631,7 @@ void FAllocator2D::CreateUsedSegments()
 				{
 					const uint32 SubX = x + i;
 
-					if ( GetBit(SubX, y) )
+					if ( BitsValue & ( 1ull << i ) )
 					{
 						if ( FirstUsedX < 0 )
 						{
@@ -416,15 +653,56 @@ void FAllocator2D::CreateUsedSegments()
 			}
 		}
 	}
+
+	SortedRowsIndex.Sort([this](uint32 a, uint32 b) { return Rows[b].LongestSegment < Rows[a].LongestSegment;});
+
+	// Create segments along the Y axis for each columns
+	if (LayoutVersion >= ELightmapUVVersion::Segments2D)
+	{
+		SortedColumnsIndex.SetNum(RasterWidth, false);
+
+		for ( uint32 x = 0; x < RasterWidth; ++x )
+		{
+			SortedColumnsIndex[x] = x;
+
+			FRun& CurrentColumn = Columns[x];
+			CurrentColumn.LongestSegment = 0;
+			CurrentColumn.Segments.Reset();
+
+			int32 FirstUsedY = -1;
+			for ( uint32 y = 0; y < RasterHeight; ++y )
+			{
+				if (GetBit(x, y))
+				{
+					if ( FirstUsedY < 0 )
+					{
+						FirstUsedY = y;
+					}
+
+					if ( y == RasterHeight - 1 )
+					{
+						AddUsedSegment( CurrentColumn, FirstUsedY, y + 1 - FirstUsedY );
+						FirstUsedY = -1;
+					}
+				}
+				else if ( FirstUsedY >= 0 )
+				{
+					AddUsedSegment( CurrentColumn, FirstUsedY, y - FirstUsedY );
+					FirstUsedY = -1;
+				}
+			}
+		}
+
+		SortedColumnsIndex.Sort([this](uint32 a, uint32 b) { return Columns[b].LongestSegment < Columns[a].LongestSegment;});
+	}
 }
 
-void FAllocator2D::AddUsedSegment( FRow& Row, uint32 StartPos, uint32 Length )
+void FAllocator2D::AddUsedSegment( FRun& Row, uint32 StartPos, uint32 Length )
 {
 	FSegment UsedSegment;
 	UsedSegment.StartPos = StartPos;
 	UsedSegment.Length = Length;
-
-	Row.UsedSegments.Add(UsedSegment);
+	Row.Segments.Add(UsedSegment);
 
 	if ( UsedSegment.Length > Row.LongestSegment )
 	{
@@ -432,63 +710,173 @@ void FAllocator2D::AddUsedSegment( FRow& Row, uint32 StartPos, uint32 Length )
 	}
 }
 
-void FAllocator2D::MergeSegments( FRect Rect, const FAllocator2D& Other )
+void FAllocator2D::MergeRun(FRun& ThisRun, const FRun& OtherRun, uint32 RectOffset, uint32 RectLength, uint32 PrimaryResolution, uint32 PerpendicularResolution)
 {
-	for ( uint32 y = 0; y < Rect.H; ++y )
+	// #TODO
+	// Perform a linear scan once for both array at the same time zipping them together
+	// Ensure we keep the order of the segments while merging them to avoid costly sort when number of free segments is high
+	for ( const FSegment& OtherUsedSegment : OtherRun.Segments )
 	{
-		FRow& ThisRow = Rows[Rect.Y + y];
-		const FRow& OtherRow = Other.Rows[y];
-
-		for ( const FSegment& OtherUsedSegment : OtherRow.UsedSegments )
+		for ( int32 Index = 0; Index < ThisRun.Segments.Num(); ++Index )
 		{
-			for ( int32 i = 0; i < ThisRow.FreeSegments.Num(); ++i )
-			{
-				FSegment& ThisFreeSegment = ThisRow.FreeSegments[i];
+			FSegment& ThisFreeSegment = ThisRun.Segments[Index];
 
-				uint32 StartPos = Rect.X + OtherUsedSegment.StartPos;
+			uint32 StartPos = RectOffset + OtherUsedSegment.StartPos;
 
-				if ( StartPos >= ThisFreeSegment.StartPos &&
-					StartPos < ThisFreeSegment.StartPos + ThisFreeSegment.Length )
+			if ( StartPos >= ThisFreeSegment.StartPos &&
+				StartPos < ThisFreeSegment.StartPos + ThisFreeSegment.Length )
+			{
+				if ( ThisFreeSegment.Length == 1 )
 				{
-					if ( ThisFreeSegment.Length == 1 )
-					{
-						ThisRow.FreeSegments.RemoveAtSwap(i);
-						break;
-					}
+					ThisRun.Segments.RemoveAtSwap(Index);
+					break;
+				}
 
-					FSegment FirstSegment;
-					FirstSegment.StartPos = ThisFreeSegment.StartPos;
-					FirstSegment.Length = StartPos - ThisFreeSegment.StartPos;
+				FSegment FirstSegment;
+				FirstSegment.StartPos = ThisFreeSegment.StartPos;
+				FirstSegment.Length = StartPos - ThisFreeSegment.StartPos;
 
-					uint32 EndPos = Rect.X + FMath::Min( OtherUsedSegment.StartPos + OtherUsedSegment.Length, Rect.W ) - 1;
+				uint32 EndPos = RectOffset + FMath::Min( OtherUsedSegment.StartPos + OtherUsedSegment.Length, RectLength ) - 1;
 
-					FSegment SecondSegment;
-					SecondSegment.StartPos = EndPos + 1;
-					SecondSegment.Length = ThisFreeSegment.StartPos + ThisFreeSegment.Length - SecondSegment.StartPos;
+				FSegment SecondSegment;
+				SecondSegment.StartPos = EndPos + 1;
+				SecondSegment.Length = ThisFreeSegment.StartPos + ThisFreeSegment.Length - SecondSegment.StartPos;
 
-					ThisRow.FreeSegments.RemoveAtSwap(i);
+				ThisRun.Segments.RemoveAtSwap(Index);
 
-					if ( FirstSegment.Length > 0 )
-					{
-						ThisRow.FreeSegments.Add( FirstSegment );
-					}
-
-					if ( SecondSegment.Length > 0 )
-					{
-						ThisRow.FreeSegments.Add( SecondSegment );
-					}
+				if ( FirstSegment.Length > 0 )
+				{
+					ThisRun.Segments.Add( FirstSegment );
+				}
 
-					break;
+				if ( SecondSegment.Length > 0 )
+				{
+					ThisRun.Segments.Add( SecondSegment );
 				}
+
+				break;
 			}
 		}
+	}
+
+	ThisRun.Segments.Sort();
+
+	// When we're dealing with an Axis, the lookup table is the size of the perpendicular Axis
+	ThisRun.FreeSegmentsLookup.SetNum(PerpendicularResolution);
+	ThisRun.LongestSegment = 0;
+
+	uint32 LastIndex = 0;
+
+	// Avoid RangeCheck by indexing the pointer directly since we're in a hot loop
+	FSegment* SegmentsPtr           = ThisRun.Segments.GetData();
+	uint16*   FreeSegmentsLookupPtr = ThisRun.FreeSegmentsLookup.GetData();
+
+	int32 SegmentIndex = 0;
+	for ( ; SegmentIndex < ThisRun.Segments.Num(); ++SegmentIndex )
+	{
+		FSegment& Segment = SegmentsPtr[SegmentIndex];
+
+		if (Segment.Length > ThisRun.LongestSegment)
+		{
+			ThisRun.LongestSegment = Segment.Length;
+		}
 
-		ThisRow.FreeSegments.Sort();
+		uint32 StopIndex  = Segment.StartPos + Segment.Length;
+		for (uint32 LookupIndex = LastIndex; LookupIndex < StopIndex; ++LookupIndex)
+		{
+			FreeSegmentsLookupPtr[LookupIndex] = SegmentIndex;
+		}
 
-		ThisRow.LongestSegment = 0;
-		for ( const FSegment& Segment : ThisRow.FreeSegments )
+		LastIndex = StopIndex;
+	}
+
+	for (uint32 Index = LastIndex; Index < PerpendicularResolution; ++Index)
+	{
+		FreeSegmentsLookupPtr[Index] = SegmentIndex;
+	}
+}
+
+void FAllocator2D::MergeSegments( const FRect& Rect, const FAllocator2D& Other )
+{
+	check(Mode       == EMode::FreeSegments);
+	check(Other.Mode == EMode::UsedSegments);
+
+	for ( uint32 y = 0; y < Other.RasterHeight; ++y )
+	{
+		MergeRun(Rows[Rect.Y + y], Other.Rows[y], Rect.X, Rect.W, Width, Height);
+	}
+
+	if (LayoutVersion >= ELightmapUVVersion::Segments2D)
+	{
+		for ( uint32 x = 0; x < Other.RasterWidth; ++x )
 		{
-			ThisRow.LongestSegment = FMath::Max( Segment.Length, ThisRow.LongestSegment );
+			MergeRun(Columns[Rect.X + x], Other.Columns[x], Rect.Y, Rect.H, Height, Width);
 		}
 	}
 }
+
+void FAllocator2D::ResetStats()
+{
+#if DEBUG_LAYOUT_STATS
+	Stats.Reset();
+#endif
+}
+
+void FAllocator2D::PublishStats(int32 ChartIndex, int32 Orientation, bool bFound, const FRect& Rect, const FRect& BestRect, const FMD5Hash& ChartMD5)
+{
+	//This is super helpful in Insights to inspect long running charts behavior
+
+	//Ensure we compile it even when not used so the code doesn't rot
+#if DEBUG_LAYOUT_STATS
+	const bool bTrace = true;
+#else
+	const bool bTrace = false;
+#endif
+
+	if (bTrace)
+	{
+		TRACE_BOOKMARK(
+			TEXT(
+				"C:%d O:%d R:%d X:%d Y:%d W:%d H:%d\n"
+				"MD5: %s\n"
+				"BestRect  X:%d Y:%d W:%d H:%d\n"
+				"BestRectBeaten %d\n"
+				"FindWithSegmentsIterationsY %llu\n"
+				"FindWithSegmentsIterationsX %llu\n"
+				"FindWithSegmentsPastBestRect %llu\n"
+				"TestAllRowsIterationsY %llu\n"
+				"FreeSegmentLookupCount %llu\n"
+				"FreeSegmentRangeIterations %llu\n"
+				"FreeSegmentFutureIterations %llu\n"
+				"FreeSegmentFutureHit %llu\n"
+				"FreeSegmentFutureHitStep %llu\n"
+				"FreeSegmentFutureMiss %llu\n"
+				"FreeSegmentFutureMissStep %llu\n"
+			),
+			ChartIndex,
+			Orientation,
+			bFound ? 1 : 0,
+			Rect.X,
+			Rect.Y,
+			Rect.W,
+			Rect.H,
+			*LexToString(ChartMD5),
+			BestRect.X,
+			BestRect.Y,
+			BestRect.W,
+			BestRect.H,
+			(bFound && ((Rect.X + Rect.Y * Width) < (BestRect.X + BestRect.Y * Width))) ? 1 : 0,
+			Stats.FindWithSegmentsIterationsY.GetValue(),
+			Stats.FindWithSegmentsIterationsX.GetValue(),
+			Stats.FindWithSegmentsMovedPastPreviousBest.GetValue(),
+			Stats.TestAllRowsIterationsY.GetValue(),
+			Stats.FreeSegmentLookupCount.GetValue(),
+			Stats.FreeSegmentRangeIterations.GetValue(),
+			Stats.FreeSegmentFutureIterations.GetValue(),
+			Stats.FreeSegmentFutureHit.GetValue(),
+			Stats.FreeSegmentFutureHitStep.GetValue(),
+			Stats.FreeSegmentFutureMiss.GetValue(),
+			Stats.FreeSegmentFutureMissStep.GetValue()
+		);
+	}
+}
\ No newline at end of file
diff --git a/Engine/Source/Developer/MeshUtilitiesCommon/Private/LayoutUV.cpp b/Engine/Source/Developer/MeshUtilitiesCommon/Private/LayoutUV.cpp
--- a/Engine/Source/Developer/MeshUtilitiesCommon/Private/LayoutUV.cpp
+++ b/Engine/Source/Developer/MeshUtilitiesCommon/Private/LayoutUV.cpp
@@ -1,15 +1,18 @@
 // Copyright 1998-2019 Epic Games, Inc. All Rights Reserved.
 
 #include "LayoutUV.h"
+#include "Algo/IntroSort.h"
+#include "Async/Async.h"
 #include "DisjointSet.h"
 #include "OverlappingCorners.h"
-#include "Algo/IntroSort.h"
 #include "HAL/PlatformTime.h"
+#include "Misc/App.h"
+#include "Misc/SecureHash.h"
 #include "Modules/ModuleManager.h"
 
 IMPLEMENT_MODULE(FDefaultModuleImpl, MeshUtilitiesCommon)
 
-DEFINE_LOG_CATEGORY_STATIC(LogLayoutUV, Warning, All);
+DEFINE_LOG_CATEGORY_STATIC(LogLayoutUV, Log, All);
 
 #define CHART_JOINING	1
 
@@ -17,6 +20,10 @@ DEFINE_LOG_CATEGORY_STATIC(LogLayoutUV, Warning, All);
 #define LEGACY_UVS_ARE_SAME (1.0f / 1024.0f)
 #define UVLAYOUT_THRESH_UVS_ARE_SAME (GetUVEqualityThreshold())
 
+TAtomic<uint64> FLayoutUV::FindBestPackingCount(0);
+TAtomic<uint64> FLayoutUV::FindBestPackingCycles(0);
+TAtomic<uint64> FLayoutUV::FindBestPackingEfficiency(0);
+
 FLayoutUV::FLayoutUV( IMeshView& InMeshView )
 	: MeshView( InMeshView )
 	, LayoutVersion( ELightmapUVVersion::Latest )
@@ -57,18 +64,15 @@ struct FLayoutUV::FChartPacker
 
 private:
 	void ScaleCharts( TArray< FMeshChart >& Charts, float UVScale );
-	bool PackCharts(TArray< FMeshChart >& Charts, const TArray< FVector2D >& TexCoords, const TArray< uint32 >& SortedTris);
+	bool PackCharts( TArray< FMeshChart >& Charts, float UVScale, const TArray< FVector2D >& TexCoords, const TArray< uint32 >& SortedTris, float& OutEfficiency, TAtomic<bool>& bAbort, bool bTrace);
 	void OrientChart( FMeshChart& Chart, int32 Orientation );
-	void RasterizeChart( const FMeshChart& Chart, const TArray< FVector2D >& TexCoords, const TArray< uint32 >& SortedTris, uint32 RectW, uint32 RectH );
+	void RasterizeChart( const FMeshChart& Chart, const TArray< FVector2D >& TexCoords, const TArray< uint32 >& SortedTris, uint32 RectW, uint32 RectH, FAllocator2D& OutChartRaster );
 
 private:
 	IMeshView& MeshView;
 	ELightmapUVVersion LayoutVersion;
 
 	uint32 TextureResolution;
-	FAllocator2D LayoutRaster;
-	FAllocator2D ChartRaster;
-	FAllocator2D BestChartRaster;
 	float TotalUVArea;
 };
 
@@ -76,9 +80,6 @@ FLayoutUV::FChartPacker::FChartPacker(IMeshView& InMeshView, ELightmapUVVersion
 	: MeshView(InMeshView)
 	, LayoutVersion(InLayoutVersion)
 	, TextureResolution(TextureResolution)
-	, LayoutRaster( TextureResolution, TextureResolution )
-	, ChartRaster( TextureResolution, TextureResolution )
-	, BestChartRaster( TextureResolution, TextureResolution )
 	, TotalUVArea(0.0f)
 {
 }
@@ -624,11 +625,19 @@ int32 FLayoutUV::FChartFinder::FindCharts( const FOverlappingCorners& Overlappin
 
 	double End = FPlatformTime::Seconds();
 
-	UE_LOG(LogLayoutUV, Display, TEXT("FindCharts: %s"), *FPlatformTime::PrettyTime(End - Begin) );
+	UE_LOG(LogLayoutUV, VeryVerbose, TEXT("FindCharts: %s"), *FPlatformTime::PrettyTime(End - Begin) );
 
 	return Charts.Num();
 }
 
+#if UE_EDITOR && (UE_BUILD_DEVELOPMENT || UE_BUILD_DEBUG)
+static TAutoConsoleVariable<FString> CVarLayoutUVTracePackingForInputHash(
+	TEXT("LayoutUV.TracePackingForInputHash"),
+	TEXT(""),
+	TEXT("Activate tracing for the input hash specified in the value.\n"),
+	ECVF_Default);
+#endif
+
 bool FLayoutUV::FChartPacker::FindBestPacking(const TArray< FVector2D >& TexCoords, const TArray< uint32 >& SortedTris, TArray< FMeshChart >& Charts)
 {
 	if( (uint32)Charts.Num() > TextureResolution * TextureResolution )
@@ -642,57 +651,219 @@ bool FLayoutUV::FChartPacker::FindBestPacking(const TArray< FVector2D >& TexCoor
 	{
 		TotalUVArea += Chart.UVArea * Chart.WorldScale.X * Chart.WorldScale.Y;
 	}
+
 	if( TotalUVArea <= 0.0f )
 	{
 		return false;
 	}
 
+	uint64 StartCycles = FPlatformTime::Cycles64();
+	TRACE_CPUPROFILER_EVENT_SCOPE(FChartPacker::FindBestPacking)
+
+	// Cleanup uninitialized values to get a stable input hash
+	for (FMeshChart& Chart : Charts)
+	{
+		Chart.PackingBias   = FVector2D::ZeroVector;
+		Chart.PackingScaleU = FVector2D::ZeroVector;
+		Chart.PackingScaleV = FVector2D::ZeroVector;
+		Chart.UVScale       = FVector2D::ZeroVector;
+	}
+
+	FString InputHash = FMD5::HashBytes((uint8*)Charts.GetData(), Charts.Num() * Charts.GetTypeSize());
+
+#if UE_EDITOR && (UE_BUILD_DEVELOPMENT || UE_BUILD_DEBUG)
+	// When you need to find where an unexpected difference in output hash might come from
+	// after changing the algorithm. You can set this CVar to activate tracing for a particular
+	// input hash.
+	FString InputHashTrace = CVarLayoutUVTracePackingForInputHash.GetValueOnAnyThread().TrimStartAndEnd();
+	const bool bTrace = InputHashTrace.Len() > 0 && InputHash.StartsWith(InputHashTrace);
+#else
+	const bool bTrace = false;
+#endif
+
+	// Those might require tuning, changing them won't affect the outcome and will maintain backward compatibility
+	const int32 MultithreadChartsCountThreshold       = 100*1000;
+	const int32 MultithreadTextureResolutionThreshold = 1000;
+	const int32 MultithreadAheadWorkCount             = 3;
+
 	const float LinearSearchStart = 0.5f;
-	const float LinearSearchStep = 0.5f;
+	const float LinearSearchStep  = 0.5f;
 	const int32 BinarySearchSteps = 6;
-	
+
 	float UVScaleFail = TextureResolution * FMath::Sqrt( 1.0f / TotalUVArea );
 	float UVScalePass = TextureResolution * FMath::Sqrt( LinearSearchStart / TotalUVArea );
 
+	// Store successful charts packing to avoid redoing the final step
+	TArray<FMeshChart>         LastPassCharts;
+	TAtomic<bool>              bAbort(false);
+
+	struct FThreadContext
+	{
+		TArray<FMeshChart> Charts;
+		TFuture<bool>      Result;
+		float              Efficiency = 0.0f;
+	};
+
+	TArray<FThreadContext> ThreadContexts;
+
+	bool bShouldUseMultipleThreads =
+		 FApp::ShouldUseThreadingForPerformance() && 
+		!bTrace &&
+ 		 Charts.Num() >= MultithreadChartsCountThreshold &&
+		 TextureResolution >= MultithreadTextureResolutionThreshold;
+
+	if ( bShouldUseMultipleThreads )
+	{
+		// Do forward work only when multi-thread activated
+		ThreadContexts.SetNum(MultithreadAheadWorkCount);
+	}
+
 	// Linear search for first fit
-	while(1)
+	float LastEfficiency = 0.0f;
 	{
-		ScaleCharts( Charts, UVScalePass );
+		TRACE_CPUPROFILER_EVENT_SCOPE(LinearSearch);
 
-		bool bFit = PackCharts(Charts, TexCoords, SortedTris);
-		if( bFit )
+		while(!bAbort)
 		{
-			break;
-		}
+			// Launch forward work in other threads
+			for (int32 Index = 0; Index < ThreadContexts.Num(); ++Index)
+			{
+				ThreadContexts[Index].Charts = Charts;
+				float ThreadUVScale  = UVScalePass * FMath::Pow(LinearSearchStep, Index + 1);
+				ThreadContexts[Index].Result =
+					Async(
+						EAsyncExecution::ThreadPool,
+						[this, &ThreadContexts, &SortedTris, &TexCoords, &bAbort, ThreadUVScale, Index]()
+						{
+							TRACE_CPUPROFILER_EVENT_SCOPE(SearchStep);
+							return PackCharts(ThreadContexts[Index].Charts, ThreadUVScale, TexCoords, SortedTris, ThreadContexts[Index].Efficiency, bAbort, false);
+						}
+					);
+			}
+
+			if (bTrace)
+			{
+				UE_LOG(LogLayoutUV, Log, TEXT("[LAYOUTUV_TRACE] Scale %f"), UVScalePass);
+			}
+
+			// Process the first iteration in this thread
+			bool bFit = false;
+			{
+				TRACE_CPUPROFILER_EVENT_SCOPE(SearchStep);
+				bFit = PackCharts(Charts, UVScalePass, TexCoords, SortedTris, LastEfficiency, bAbort, bTrace);
+			}
 
-		UVScaleFail = UVScalePass;
-		UVScalePass *= LinearSearchStep;
+			// Wait for the work sequentially and cancel everything once we have a first viable solution
+			for (int32 Index = 0; Index < ThreadContexts.Num() + 1; ++Index)
+			{
+				// The first result is not coming from a future
+				bFit = Index == 0 ? bFit : ThreadContexts[Index - 1].Result.Get();
+				if (bFit && !bAbort)
+				{
+					// We got a success, cancel other searches
+					bAbort = true;
+
+					if (Index > 0)
+					{
+						Charts         = ThreadContexts[Index - 1].Charts;
+						LastEfficiency = ThreadContexts[Index - 1].Efficiency;
+					}
+
+					LastPassCharts = Charts;
+				}
+
+				if (!bAbort)
+				{
+					UVScaleFail = UVScalePass;
+					UVScalePass *= LinearSearchStep;
+				}
+			}
+		}
 	}
 
 	// Binary search for best fit
-	for( int32 i = 0; i < BinarySearchSteps; i++ )
 	{
-		float UVScale = 0.5f * ( UVScaleFail + UVScalePass );
-		ScaleCharts( Charts, UVScale );
+		TRACE_CPUPROFILER_EVENT_SCOPE(BinarySearch);
 
-		bool bFit = PackCharts(Charts, TexCoords, SortedTris);
-		if( bFit )
+		bAbort = false;
+		for( int32 i = 0; i < BinarySearchSteps; i++ )
 		{
-			UVScalePass = UVScale;
-		}
-		else
-		{
-			UVScaleFail = UVScale;
+			TRACE_CPUPROFILER_EVENT_SCOPE(SearchStep);
+
+			float UVScale = 0.5f * ( UVScaleFail + UVScalePass );
+
+			if (bTrace)
+			{
+				UE_LOG(LogLayoutUV, Log, TEXT("[LAYOUTUV_TRACE] Scale %f"), UVScale);
+			}
+
+			float Efficiency = 0.0f;
+			bool bFit = PackCharts(Charts, UVScale, TexCoords, SortedTris, Efficiency, bAbort, bTrace);
+			if( bFit )
+			{
+				LastPassCharts = Charts;
+				float EfficiencyGainPercent = 100.0f * FMath::Abs(Efficiency - LastEfficiency);
+				LastEfficiency = Efficiency;
+
+				// Early out when we're inside a 1% efficiency range
+				if (LayoutVersion >= ELightmapUVVersion::Segments2D && EfficiencyGainPercent <= 1.0f)
+				{
+					break;
+				}
+			
+				UVScalePass = UVScale;
+			}
+			else
+			{
+				UVScaleFail = UVScale;
+			}
 		}
 	}
 
-	// TODO store packing scale/bias separate so this isn't necessary
-	ScaleCharts( Charts, UVScalePass );
-	PackCharts(Charts, TexCoords, SortedTris);
+	if (LayoutVersion < ELightmapUVVersion::ScaleChartsOrderingFix)
+	{
+		// Early versions applied a sort that was determinist
+		// but dependent on earlier sorts. Since we strive to maintain
+		// backward compatibility of the UV layout to avoid screwing
+		// with already backed static lighting, we must apply a final
+		// scaling and packing that will reuse the last step's ordering
+		// whether it was a failure or not.
+		PackCharts(Charts, UVScalePass, TexCoords, SortedTris, LastEfficiency, bAbort, bTrace);
+	}
+	else
+	{
+		// In case the last step was a failure, restore from last known good computation
+		Charts = LastPassCharts;
+	}
+
+	FString OutputHash = FMD5::HashBytes((uint8*)Charts.GetData(), Charts.Num() * Charts.GetTypeSize());
+
+	// Increase verbosity level to use this for packing results validation when modifying code
+	UE_LOG(LogLayoutUV, Verbose, TEXT("FindBestPacking (Input Data MD5: %s, Output Data MD5: %s, LayoutVersion: %d, Efficiency: %0.2f %%)"), *InputHash, *OutputHash, LayoutVersion, LastEfficiency*100);
+
+	static TAtomic<uint64> Count(0);
+	static TAtomic<uint64> TotalCycles(0);
+	static TAtomic<uint64> Efficiency(0);
+
+	FindBestPackingCount++;
+	FindBestPackingEfficiency += LastEfficiency*100000;
+	FindBestPackingCycles += FPlatformTime::Cycles64() - StartCycles;
 
 	return true;
 }
 
+void FLayoutUV::ResetStats()
+{
+	FindBestPackingCount = 0;
+	FindBestPackingEfficiency = 0;
+	FindBestPackingCycles = 0;
+}
+
+void FLayoutUV::LogStats()
+{
+	UE_LOG(LogLayoutUV, Log, TEXT("FindBestPacking (Total Time: %s, Avg Efficiency: %f)"), *FPlatformTime::PrettyTime(FPlatformTime::ToSeconds64(FindBestPackingCycles.Load())), double(FindBestPackingEfficiency.Load()) / (FindBestPackingCount.Load()*1000));
+}
+
 void FLayoutUV::FChartPacker::ScaleCharts( TArray< FMeshChart >& Charts, float UVScale )
 {
 	for( int32 i = 0; i < Charts.Num(); i++ )
@@ -838,137 +1009,241 @@ void FLayoutUV::FChartPacker::ScaleCharts( TArray< FMeshChart >& Charts, float U
 	Algo::IntroSort( Charts, FCompareCharts() );
 }
 
-bool FLayoutUV::FChartPacker::PackCharts(TArray< FMeshChart >& Charts, const TArray< FVector2D >& TexCoords, const TArray< uint32 >& SortedTris)
+// Hash function to use FMD5Hash in TMap
+inline uint32 GetTypeHash(const FMD5Hash& Hash)
 {
-	uint32 RasterizeCycles = 0;
-	uint32 FindCycles = 0;
-	
+	uint32* HashAsInt32 = (uint32*)Hash.GetBytes();
+	return HashAsInt32[0] ^ HashAsInt32[1] ^ HashAsInt32[2] ^ HashAsInt32[3];
+}
+
+bool FLayoutUV::FChartPacker::PackCharts(TArray< FMeshChart >& Charts, float UVScale, const TArray< FVector2D >& TexCoords, const TArray< uint32 >& SortedTris, float& OutEfficiency, TAtomic<bool>& bAbort, bool bTrace)
+{
+	ScaleCharts( Charts, UVScale );
+	TRACE_CPUPROFILER_EVENT_SCOPE(FChartPacker::PackCharts)
+
+	FAllocator2D BestChartRaster(FAllocator2D::EMode::UsedSegments, TextureResolution, TextureResolution, LayoutVersion);
+	FAllocator2D ChartRaster    (FAllocator2D::EMode::UsedSegments, TextureResolution, TextureResolution, LayoutVersion);
+	FAllocator2D LayoutRaster   (FAllocator2D::EMode::FreeSegments, TextureResolution, TextureResolution, LayoutVersion);
+
+	uint64 RasterizeCycles = 0;
+	uint64 FindCycles = 0;
+
 	double BeginPackCharts = FPlatformTime::Seconds();
 
+	OutEfficiency = 0.0f;
 	LayoutRaster.Clear();
 
-	for( int32 i = 0; i < Charts.Num(); i++ )
+	// Store the position where we found a spot for each unique raster
+	// so we can skip whole sections we know won't work out.
+	// This method is obviously more efficient with smaller charts
+	// but helps tremendously as the number of charts goes up for
+	// the same texture space. This helps counteract the slowdown
+	// induced by having more parts to place in the grid and is
+	// particularly useful for foliage.
+	TMap<FMD5Hash, FVector2D> BestStartPos;
+
+	// Reduce Insights CPU tracing to once per batch
+	const int32 BatchSize = 1024;
+	for( int32 ChartIndex = 0; ChartIndex < Charts.Num() && !bAbort.Load(EMemoryOrder::Relaxed);)
 	{
-		FMeshChart& Chart = Charts[i];
+		TRACE_CPUPROFILER_EVENT_SCOPE(ChartBatch);
+		for( int32 BatchIndex = 0; BatchIndex < BatchSize && ChartIndex < Charts.Num() && !bAbort.Load(EMemoryOrder::Relaxed); ++ChartIndex, ++BatchIndex)
+		{
+			FMeshChart& Chart = Charts[ChartIndex];
 
-		// Try different orientations and pick best
-		int32				BestOrientation = -1;
-		FAllocator2D::FRect	BestRect = { ~0u, ~0u, ~0u, ~0u };
+			// Try different orientations and pick best
+			int32				BestOrientation = -1;
+			FAllocator2D::FRect	BestRect = { ~0u, ~0u, ~0u, ~0u };
 
-		for( int32 Orientation = 0; Orientation < 8; Orientation++ )
-		{
-			// TODO If any dimension is less than 1 pixel shrink dimension to zero
+			for( int32 Orientation = 0; Orientation < 8; Orientation++ )
+			{
+				// TODO If any dimension is less than 1 pixel shrink dimension to zero
 
-			OrientChart( Chart, Orientation );
+				OrientChart( Chart, Orientation);
 			
-			FVector2D ChartSize = Chart.MaxUV - Chart.MinUV;
-			ChartSize = ChartSize.X * Chart.PackingScaleU + ChartSize.Y * Chart.PackingScaleV;
+				FVector2D ChartSize = Chart.MaxUV - Chart.MinUV;
+				ChartSize = ChartSize.X * Chart.PackingScaleU + ChartSize.Y * Chart.PackingScaleV;
 
-			// Only need half pixel dilate for rects
-			FAllocator2D::FRect	Rect;
-			Rect.X = 0;
-			Rect.Y = 0;
-			Rect.W = FMath::CeilToInt( FMath::Abs( ChartSize.X ) + 1.0f );
-			Rect.H = FMath::CeilToInt( FMath::Abs( ChartSize.Y ) + 1.0f );
+				// Only need half pixel dilate for rects
+				FAllocator2D::FRect	Rect;
+				Rect.X = 0;
+				Rect.Y = 0;
+				Rect.W = FMath::CeilToInt( FMath::Abs( ChartSize.X ) + 1.0f );
+				Rect.H = FMath::CeilToInt( FMath::Abs( ChartSize.Y ) + 1.0f );
 
-			// Just in case lack of precision pushes it over
-			Rect.W = FMath::Min( TextureResolution, Rect.W );
-			Rect.H = FMath::Min( TextureResolution, Rect.H );
+				// Just in case lack of precision pushes it over
+				Rect.W = FMath::Min( TextureResolution, Rect.W );
+				Rect.H = FMath::Min( TextureResolution, Rect.H );
 
-			const bool bRectPack = false;
+				const bool bRectPack = false;
 
-			if( bRectPack )
-			{
-				if( LayoutRaster.Find( Rect ) )
+				if( bRectPack )
 				{
-					// Is best?
-					if( Rect.X + Rect.Y * TextureResolution < BestRect.X + BestRect.Y * TextureResolution )
+					if( LayoutRaster.Find( Rect ) )
 					{
-						BestOrientation = Orientation;
-						BestRect = Rect;
+						// Is best?
+						if( Rect.X + Rect.Y * TextureResolution < BestRect.X + BestRect.Y * TextureResolution )
+						{
+							BestOrientation = Orientation;
+							BestRect = Rect;
+						}
+					}
+					else
+					{
+						continue;
 					}
 				}
 				else
 				{
-					continue;
-				}
-			}
-			else
-			{
-				if ( LayoutVersion >= ELightmapUVVersion::Segments && Orientation % 4 == 1 )
-				{
-					ChartRaster.FlipX( Rect, LayoutVersion );
-				}
-				else if ( LayoutVersion >= ELightmapUVVersion::Segments && Orientation % 4 == 3 )
-				{
-					ChartRaster.FlipY( Rect );
-				}
-				else
-				{
-					int32 BeginRasterize = FPlatformTime::Cycles();
-					RasterizeChart( Chart, TexCoords, SortedTris, Rect.W , Rect.H );
-					RasterizeCycles += FPlatformTime::Cycles() - BeginRasterize;
-				}
-
-				bool bFound = false;
+					if ( LayoutVersion >= ELightmapUVVersion::Segments && Orientation % 4 == 1 )
+					{
+						ChartRaster.FlipX( Rect );
+					}
+					else if ( LayoutVersion >= ELightmapUVVersion::Segments && Orientation % 4 == 3 )
+					{
+						ChartRaster.FlipY( Rect );
+					}
+					else
+					{
+						int32 BeginRasterize = FPlatformTime::Cycles();
+						RasterizeChart( Chart, TexCoords, SortedTris, Rect.W , Rect.H, ChartRaster);
+						RasterizeCycles += FPlatformTime::Cycles() - BeginRasterize;
+					}
 
-				uint32 BeginFind = FPlatformTime::Cycles();
-				if ( LayoutVersion == ELightmapUVVersion::BitByBit )
-				{
-					bFound = LayoutRaster.FindBitByBit( Rect, ChartRaster );
-				}
-				else if ( LayoutVersion >= ELightmapUVVersion::Segments )
-				{
-					bFound = LayoutRaster.FindWithSegments( Rect, BestRect, ChartRaster );
-				}
-				FindCycles += FPlatformTime::Cycles() - BeginFind;
+					bool bFound = false;
 
-				if( bFound )
-				{
-					// Is best?
-					if( Rect.X + Rect.Y * TextureResolution < BestRect.X + BestRect.Y * TextureResolution )
+					uint32 BeginFind = FPlatformTime::Cycles();
+					if ( LayoutVersion == ELightmapUVVersion::BitByBit )
+					{
+						bFound = LayoutRaster.FindBitByBit( Rect, ChartRaster );
+					}
+					else if ( LayoutVersion >= ELightmapUVVersion::Segments )
 					{
-						BestChartRaster = ChartRaster;
+						// Use the real raster size for optimal placement
+						FAllocator2D::FRect RasterRect = Rect;
+						RasterRect.W = ChartRaster.GetRasterWidth();
+						RasterRect.H = ChartRaster.GetRasterHeight();
+
+						// Nothing rasterized, returning 0,0 as fast as possible
+						// since this is what the actual algorithm is doing but
+						// we might have to flag the entire UV map as invalid since
+						// charts are going to overlap
+						if (RasterRect.H == 0 && RasterRect.W == 0)
+						{
+							Rect.X = 0;
+							Rect.Y = 0;
+							bFound = true;
+						}
+						else
+						{
+							FMD5Hash RasterMD5 = ChartRaster.GetRasterMD5();
+							FVector2D* StartPos = BestStartPos.Find(RasterMD5);
+
+							if (StartPos)
+							{
+								RasterRect.X = StartPos->X;
+								RasterRect.Y = StartPos->Y;
+							}
+
+							LayoutRaster.ResetStats();
+							bFound = LayoutRaster.FindWithSegments(RasterRect, BestRect, ChartRaster);
+							if (bFound)
+							{
+								// Store only the best possible position in the hash table so we can start from there for other identical charts
+								BestStartPos.Add(RasterMD5, FVector2D(RasterRect.X, RasterRect.Y));
+
+								// Since the older version stops searching at Width - Rect.W instead of using the raster size,
+								// it means a perfect rasterized square of 2,2 won't fit a 2,2 hole at the end of a row if Rect.W = 3.
+								// Because of that, we have no choice to worsen our algorithm behavior for backward compatibility.
+
+								// Once we know the best possible position, we'll continue our search from there with the original
+								// rect value if it differs from the raster rect to ensure we get the same result as the old algorithm.
+								if (LayoutVersion < ELightmapUVVersion::Segments2D && (Rect.X != RasterRect.X || Rect.Y != RasterRect.Y))
+								{
+									Rect.X = RasterRect.X;
+									Rect.Y = RasterRect.Y;
+
+									bFound = LayoutRaster.FindWithSegments(Rect, BestRect, ChartRaster);
+								}
+								else
+								{
+									// We can't copy W and H here as they might be different than what we got initially
+									Rect.X = RasterRect.X;
+									Rect.Y = RasterRect.Y;
+								}
+							}
+
+							LayoutRaster.PublishStats(ChartIndex, Orientation, bFound, Rect, BestRect, RasterMD5);
+						}
 
-						BestOrientation = Orientation;
-						BestRect = Rect;
+					}
+					FindCycles += FPlatformTime::Cycles() - BeginFind;
+
+					if (bTrace)
+					{
+						UE_LOG(LogLayoutUV, Log, TEXT("[LAYOUTUV_TRACE] Chart %d Orientation %d Found = %d Rect = %d,%d,%d,%d\n"), ChartIndex, Orientation, bFound ? 1 : 0, Rect.X, Rect.Y, Rect.W, Rect.H);
+					}
 
-						if ( BestRect.X == 0 && BestRect.Y == 0 )
+					if( bFound )
+					{
+						// Is best?
+						if( Rect.X + Rect.Y * TextureResolution < BestRect.X + BestRect.Y * TextureResolution )
 						{
-							// BestRect can't be beat, stop here
-							break;
+							BestChartRaster = ChartRaster;
+
+							BestOrientation = Orientation;
+							BestRect = Rect;
+
+							if ( BestRect.X == 0 && BestRect.Y == 0 )
+							{
+								// BestRect can't be beat, stop here
+								break;
+							}
 						}
 					}
-				}
-				else
-				{
-					continue;
+					else
+					{
+						continue;
+					}
 				}
 			}
-		}
 
-		if( BestOrientation >= 0 )
-		{
-			// Add chart to layout
-			OrientChart( Chart, BestOrientation );
+			if( BestOrientation >= 0 )
+			{
+				// Add chart to layout
+				OrientChart( Chart, BestOrientation );
 
-			LayoutRaster.Alloc( BestRect, BestChartRaster );
+				LayoutRaster.Alloc( BestRect, BestChartRaster );
 
-			Chart.PackingBias.X += BestRect.X;
-			Chart.PackingBias.Y += BestRect.Y;
-		}
-		else
-		{
-			// Found no orientation that fit
-			return false;
+				Chart.PackingBias.X += BestRect.X;
+				Chart.PackingBias.Y += BestRect.Y;
+			}
+			else
+			{
+				if (bTrace)
+				{
+					UE_LOG(LogLayoutUV, Log, TEXT("[LAYOUTUV_TRACE] Chart %d Found no orientation that fit\n"), ChartIndex);
+				}
+				// Found no orientation that fit
+				return false;
+			}
 		}
 	}
 
+	if (bAbort)
+	{
+		return false;
+	}
+
+	const uint32 TotalTexels = TextureResolution * TextureResolution;
+	const uint32 UsedTexels  = LayoutRaster.GetUsedTexels();
+
+	OutEfficiency = float( UsedTexels ) / TotalTexels;
 	double EndPackCharts = FPlatformTime::Seconds();
 
-	UE_LOG(LogLayoutUV, Display, TEXT("PackCharts: %s"),	*FPlatformTime::PrettyTime(EndPackCharts - BeginPackCharts));
-	UE_LOG(LogLayoutUV, Display, TEXT("  Rasterize: %u"),	RasterizeCycles);
-	UE_LOG(LogLayoutUV, Display, TEXT("  Find: %u"),		FindCycles);
+	UE_LOG(LogLayoutUV, VeryVerbose, TEXT("PackCharts: %s"),	*FPlatformTime::PrettyTime(EndPackCharts - BeginPackCharts));
+	UE_LOG(LogLayoutUV, VeryVerbose, TEXT("  Rasterize: %llu"),	RasterizeCycles);
+	UE_LOG(LogLayoutUV, VeryVerbose, TEXT("  Find: %llu"),		FindCycles);
 
 	return true;
 }
@@ -1109,14 +1384,14 @@ void RasterizeTriangle( FAllocator2D& Shader, const FVector2D Points[3], int32 S
 	}
 }
 
-void FLayoutUV::FChartPacker::RasterizeChart( const FMeshChart& Chart, const TArray< FVector2D >& TexCoords, const TArray< uint32 >& SortedTris, uint32 RectW, uint32 RectH )
+void FLayoutUV::FChartPacker::RasterizeChart( const FMeshChart& Chart, const TArray< FVector2D >& TexCoords, const TArray< uint32 >& SortedTris, uint32 RectW, uint32 RectH, FAllocator2D& OutChartRaster )
 {
 	// Bilinear footprint is -1 to 1 pixels. If packed geometrically, only a half pixel dilation
 	// would be needed to guarantee all charts were at least 1 pixel away, safe for bilinear filtering.
 	// Unfortunately, with pixel packing a full 1 pixel dilation is required unless chart edges exactly
 	// align with pixel centers.
 
-	ChartRaster.Clear();
+	OutChartRaster.Clear();
 
 	for( uint32 Tri = Chart.FirstTri; Tri < Chart.LastTri; Tri++ )
 	{
@@ -1127,12 +1402,12 @@ void FLayoutUV::FChartPacker::RasterizeChart( const FMeshChart& Chart, const TAr
 			Points[k] = UV.X * Chart.PackingScaleU + UV.Y * Chart.PackingScaleV + Chart.PackingBias;
 		}
 
-		RasterizeTriangle< 16 >( ChartRaster, Points, RectW, RectH );
+		RasterizeTriangle< 16 >( OutChartRaster, Points, RectW, RectH );
 	}
 
 	if ( LayoutVersion >= ELightmapUVVersion::Segments )
 	{
-		ChartRaster.CreateUsedSegments();
+		OutChartRaster.CreateUsedSegments();
 	}
 }
 
@@ -1151,6 +1426,8 @@ void FLayoutUV::CommitPackedUVs()
 		return;
 	}
 
+	TRACE_CPUPROFILER_EVENT_SCOPE(FLayoutUV::CommitPackedUVs)
+
 	// Alloc new UV channel
 	MeshView.InitOutputTexcoords(MeshTexCoords.Num());
 
diff --git a/Engine/Source/Developer/MeshUtilitiesCommon/Public/Allocator2D.h b/Engine/Source/Developer/MeshUtilitiesCommon/Public/Allocator2D.h
--- a/Engine/Source/Developer/MeshUtilitiesCommon/Public/Allocator2D.h
+++ b/Engine/Source/Developer/MeshUtilitiesCommon/Public/Allocator2D.h
@@ -7,10 +7,23 @@
 
 struct FRect;
 struct Rect;
+struct FMD5Hash;
+
+#define DEBUG_LAYOUT_STATS 0
 
 class MESHUTILITIESCOMMON_API FAllocator2D
 {
 public:
+	enum class EMode
+	{
+		// In this mode, segments represents free space
+		// Used for the layout merging usedsegments
+		FreeSegments,
+		// In this mode, segments represents used space
+		// Used for the rasterization of charts.
+		UsedSegments
+	};
+
 	struct FRect
 	{
 		uint32 X;
@@ -27,56 +40,131 @@ public:
 		bool operator<( const FSegment& Other ) const { return StartPos < Other.StartPos; }
 	};
 
-	struct FRow
+	struct FRun
 	{
-		uint32 Index;
-		uint32 LongestSegment; // Represents either the longest free segment or the longest used segment, depending on how we're using this row
-
-		TArray< FSegment, TInlineAllocator<2> > FreeSegments;
-		TArray< FSegment, TInlineAllocator<2> > UsedSegments;
+		uint32 LongestSegment;
+		TArray< FSegment > Segments;
+
+		// Contains mapping from pixel position to first segment index in search range.
+		// Only computed when we're in FreeSegments mode to help TestOneRun find
+		// the proper segment in O(c) at the expense of a likely cache miss.
+		// We'll use a threshold to use this method when the number of iterations
+		// saved is worth the cache miss.
+		// Obviously, using a uint16 here will reduce cache misses but impose
+		// an hopefully enough limitation of texture size 65536x65536 (4GB).
+		TArray< uint16 > FreeSegmentsLookup;
 	};
 
 public:
-	FAllocator2D( uint32 Width, uint32 Height );
+	FAllocator2D( EMode Mode, uint32 Width, uint32 Height, ELightmapUVVersion LayoutVersion );
 
 	// Must clear before using
-	void		Clear();
+	void       Clear();
 
-	bool		Find( FRect& Rect );
-	bool		Test( FRect Rect );
-	void		Alloc( FRect Rect );
+	bool       Find( FRect& Rect );
+	bool       Test( FRect Rect );
+	void       Alloc( FRect Rect );
 
-	bool		FindBitByBit( FRect& Rect, const FAllocator2D& Other );
-	bool		FindWithSegments( FRect& Rect, FRect BestRect, const FAllocator2D& Other );
-	bool		Test( FRect Rect, const FAllocator2D& Other );
-	void		Alloc( FRect Rect, const FAllocator2D& Other );
-	
-	uint64		GetBit( uint32 x, uint32 y ) const;
-	void		SetBit( uint32 x, uint32 y );
-	void		ClearBit( uint32 x, uint32 y );
-	
-	void		CreateUsedSegments();
-	void		MergeSegments( FRect Rect, const FAllocator2D& Other );
+	bool       FindBitByBit( FRect& Rect, const FAllocator2D& Other );
+	bool       FindWithSegments( FRect& Rect, const FRect& BestRect, const FAllocator2D& Other ) const;
+	bool       Test( FRect Rect, const FAllocator2D& Other );
+	void       Alloc( FRect Rect, const FAllocator2D& Other );
 
-	void		FlipX( FRect Rect, ELightmapUVVersion LayoutVersion = ELightmapUVVersion::Latest );
-	void		FlipY( FRect Rect );
+	uint64     GetBit( uint32 x, uint32 y ) const;
+	void       SetBit( uint32 x, uint32 y );
+	void       ClearBit( uint32 x, uint32 y );
 
-protected:
-	bool		TestAllRows( FRect Rect, const FAllocator2D& Other, uint32& FailedLength );
-	bool		TestRow( const FRow& ThisRow, const FRow& OtherRow, FRect Rect, uint32& FailedLength );
+	void       CreateUsedSegments();
+	void       MergeRun( FRun& Run, const FRun& OtherRun, uint32 RectOffset, uint32 RectLength, uint32 PrimaryResolution /* Resolution along the axis the run belongs to */, uint32 PerpendicularResolution );
+	void       MergeSegments( const FRect& Rect, const FAllocator2D& Other );
 
-	void		InitSegments();
-	void		AddUsedSegment( FRow& Row, uint32 StartPos, uint32 Length );
+	void       FlipX( const FRect& Rect );
+	void       FlipY( const FRect& Rect );
 
-private:
-	TArray< uint64 > Bits;
+	uint32     GetUsedTexels() const;
+
+	void       CopyRuns( TArray<FRun>& Runs, const TArray<FRun>& OtherRuns, int32 MaxSize );
+
+	// Take control of the copy to reduce the amount of data movement to the strict minimum
+	FAllocator2D& operator = (const FAllocator2D& Other);
+
+	// Allow to visualize the content in ascii for debugging purpose. (i.e Watch or Immediate window).
+	FString    ToString() const;
 
-	uint32		Width;
-	uint32		Height;
-	uint32		Pitch;
+	// Get the MD5 hash of the rasterized content
+	FMD5Hash   GetRasterMD5() const;
+
+	uint32     GetRasterWidth()  const { return RasterWidth; }
+	uint32     GetRasterHeight() const { return RasterHeight; }
+
+	void       ResetStats();
+	void       PublishStats( int32 ChartIndex, int32 Orientation, bool bFound, const FRect& Rect, const FRect& BestRect, const FMD5Hash& ChartMD5 );
+
+protected:
+	bool       TestOneRun( const FRun& Run, const FRun& OtherRun, uint32 RectOffset, uint32 RectLength, uint32 PrimaryResolution, uint32& OutFailedLength ) const;
+	bool       TestAllRows( const FRect& Rect, const FAllocator2D& Other, uint32& FailedLength ) const;
+	bool       TestAllColumns( const FRect& Rect, const FAllocator2D& Other, uint32& FailedLength ) const;
+	void       InitRuns( TArray<FRun>& Runs, uint32 PrimaryResolution, uint32 PerpendicularRasterSize);
+	void       InitSegments();
+	void       AddUsedSegment( FRun& Run, uint32 StartPos, uint32 Length );
+
+	// Enforce that those cannot be changed in flight
+	const EMode              Mode;
+	const uint32             Width;
+	const uint32             Height;
+	const uint32             Pitch;
+	const ELightmapUVVersion LayoutVersion;
+
+	uint32                   RasterWidth;
+	uint32                   RasterHeight;
+	TArray< FRun >           Rows;        // Represent rows in the grid
+	TArray< FRun >           Columns;     // Represent columns in the grid (used when version >= Segments2D).
+	TArray< uint64 >         Bits;
+
+	// Index inside rows that will be sorted by rows with longest used segment first
+	TArray< uint16 >         SortedRowsIndex;
+	// Index inside columns that will be sorted by columns with longest used segment first
+	TArray< uint16 >         SortedColumnsIndex;
 
-	TArray< FRow > Rows;
-	int32		LastRowFail;
+private:
+	// Store iteration stats of the principal algorithms.
+	struct FStats
+	{
+		struct FStat
+		{
+#if DEBUG_LAYOUT_STATS
+			FORCEINLINE void operator++(int) { Value++; }
+			FORCEINLINE void operator+=(uint64 InValue) { Value += InValue; }
+			FORCEINLINE uint64 GetValue() const { return Value; }
+		private:
+			uint64 Value = 0;
+#else
+			// These will be optimized away when stats are not activated
+			FORCEINLINE void operator++(int) { }
+			FORCEINLINE void operator+=(uint64 InValue) { }
+			FORCEINLINE uint64 GetValue() const { return 0; }
+#endif
+		};
+
+		FStat FindWithSegmentsIterationsY;
+		FStat FindWithSegmentsIterationsX;
+		FStat FindWithSegmentsMovedPastPreviousBest;
+		FStat TestAllRowsIterationsY;
+		FStat FreeSegmentLookupCount;
+		FStat FreeSegmentRangeIterations;
+		FStat FreeSegmentFutureIterations;
+		FStat FreeSegmentFutureHit;
+		FStat FreeSegmentFutureHitStep;
+		FStat FreeSegmentFutureMiss;
+		FStat FreeSegmentFutureMissStep;
+
+		void Reset()
+		{
+			FPlatformMemory::Memzero(this, sizeof(FStats));
+		}
+	};
+
+	mutable FStats Stats;
 };
 
 // Returns non-zero if set
@@ -88,6 +176,17 @@ FORCEINLINE uint64 FAllocator2D::GetBit( uint32 x, uint32 y ) const
 FORCEINLINE void FAllocator2D::SetBit( uint32 x, uint32 y )
 {
 	Bits[ (x >> 6) + y * Pitch ] |= ( 1ull << ( x & 63 ) );
+
+	// Keep track of the rasterized dimension to optimize operations on that area only
+	if (y >= RasterHeight)
+	{
+		RasterHeight = y + 1;
+	}
+
+	if (x >= RasterWidth)
+	{
+		RasterWidth = x + 1;
+	}
 }
 
 FORCEINLINE void FAllocator2D::ClearBit( uint32 x, uint32 y )
diff --git a/Engine/Source/Developer/MeshUtilitiesCommon/Public/LayoutUV.h b/Engine/Source/Developer/MeshUtilitiesCommon/Public/LayoutUV.h
--- a/Engine/Source/Developer/MeshUtilitiesCommon/Public/LayoutUV.h
+++ b/Engine/Source/Developer/MeshUtilitiesCommon/Public/LayoutUV.h
@@ -59,6 +59,8 @@ public:
 	bool FindBestPacking( uint32 InTextureResolution );
 	void CommitPackedUVs();
 
+	static void LogStats();
+	static void ResetStats();
 private:
 	IMeshView& MeshView;
 	ELightmapUVVersion LayoutVersion;
@@ -70,4 +72,8 @@ private:
 
 	struct FChartFinder;
 	struct FChartPacker;
+
+	static TAtomic<uint64> FindBestPackingCount;
+	static TAtomic<uint64> FindBestPackingCycles;
+	static TAtomic<uint64> FindBestPackingEfficiency;
 };
diff --git a/Engine/Source/Developer/MeshUtilitiesCommon/Public/MeshUtilitiesCommon.h b/Engine/Source/Developer/MeshUtilitiesCommon/Public/MeshUtilitiesCommon.h
--- a/Engine/Source/Developer/MeshUtilitiesCommon/Public/MeshUtilitiesCommon.h
+++ b/Engine/Source/Developer/MeshUtilitiesCommon/Public/MeshUtilitiesCommon.h
@@ -14,7 +14,8 @@ enum class ELightmapUVVersion : int32
 	Allocator2DFlipFix = 5,
 	ConsiderLightmapPadding = 6,
 	ForceLightmapPadding = 7,
-	Latest = ForceLightmapPadding
+	Segments2D = 8,
+	Latest = Segments2D
 };
 
 /** Helper struct for building acceleration structures. */
@@ -80,4 +81,4 @@ namespace TriangleUtilities
 		float DotProduct = FVector::DotProduct(E1, E2);
 		return FMath::Acos(DotProduct);
 	}
-}
\ No newline at end of file
+}
