Fix CIS build

diff --git a/Engine/Source/Developer/MeshUtilitiesCommon/MeshUtilitiesCommon.Build.cs b/Engine/Source/Developer/MeshUtilitiesCommon/MeshUtilitiesCommon.Build.cs
--- a/Engine/Source/Developer/MeshUtilitiesCommon/MeshUtilitiesCommon.Build.cs
+++ b/Engine/Source/Developer/MeshUtilitiesCommon/MeshUtilitiesCommon.Build.cs
@@ -18,6 +18,7 @@ namespace UnrealBuildTool.Rules
 					"RenderCore",
 					"Engine",
                     "ClothingSystemRuntime",
+					"Slate",
                 }
             );
         }
diff --git a/Engine/Source/Developer/MeshUtilitiesCommon/Private/LODUtilities.cpp b/Engine/Source/Developer/MeshUtilitiesCommon/Private/LODUtilities.cpp
--- a/Engine/Source/Developer/MeshUtilitiesCommon/Private/LODUtilities.cpp
+++ b/Engine/Source/Developer/MeshUtilitiesCommon/Private/LODUtilities.cpp
@@ -16,6 +16,7 @@
 #include "MeshUtilities.h"
 #include "OverlappingCorners.h"
 #include "Assets/ClothingAsset.h"
+#include "Framework/Commands/UIAction.h"
 
 #include "ObjectTools.h"
 #include "ScopedTransaction.h"
diff --git a/Engine/Source/Developer/MeshUtilitiesCommon/Public/LODUtilities.h b/Engine/Source/Developer/MeshUtilitiesCommon/Public/LODUtilities.h
--- a/Engine/Source/Developer/MeshUtilitiesCommon/Public/LODUtilities.h
+++ b/Engine/Source/Developer/MeshUtilitiesCommon/Public/LODUtilities.h
@@ -6,6 +6,7 @@
 #include "MeshBuild.h"
 #include "Engine/SkeletalMesh.h"
 #include "Rendering/SkeletalMeshLODImporterData.h"
+#include "Framework/Commands/UIAction.h"
 
 
 
