Work around for a compiler bug, the original code is correct but the generated asm code would not run the range check in the while loop once after RemoveAtSwap was called and we would end up asserting/reading garbage if the last element of Charts was removed.

diff --git a/Engine/Source/Developer/MeshUtilitiesCommon/Private/LayoutUV.cpp b/Engine/Source/Developer/MeshUtilitiesCommon/Private/LayoutUV.cpp
--- a/Engine/Source/Developer/MeshUtilitiesCommon/Private/LayoutUV.cpp
+++ b/Engine/Source/Developer/MeshUtilitiesCommon/Private/LayoutUV.cpp
@@ -593,12 +593,16 @@ int32 FLayoutUV::FChartFinder::FindCharts( const FOverlappingCorners& Overlappin
 	}
 	
 	// Clean out empty charts
-	for( int32 i = 0; i < Charts.Num(); i++ )
+	for( int32 i = 0; i < Charts.Num(); )
 	{
-		while( i < Charts.Num() && Charts[i].FirstTri == Charts[i].LastTri )
+		if( Charts[i].FirstTri == Charts[i].LastTri )
 		{
 			Charts.RemoveAtSwap(i);
 		}
+		else
+		{ 
+			i++;
+		}	
 	}
 
 	for( int32 i = 0; i < Charts.Num(); i++ )
